Sequencer- First pass on Movie Scene Anim Mixer.

This experimental plugin provides a different path for processing skeletal animation within Sequencer. The main point of this is to modularize animation producing and mixing within Sequencer. The new structure allows multiple types of animation producing tracks (e.g. anim sequence, control rig, plus future types like idle, facial anim, gaze, motion matching, etc.). It also allows multiple types of animation 'targets'. All animations sent to a single animation target will be mixed (and eventually masked) together.

The actual mixing is implemented using the Anim Next 'Evaluation Task' system. A list of evaluation tasks (also called an Evaluation Program) runs in sequence, operating with a shared VM memory. This memory holds a 'pose keyframe' stack. A pose keyframe holds optionally a bone pose, evaluated curve values, and attribute values. Root motion is represented as a uniquely named attribute. Tasks can push new pose keyframes to the stack, or they can perform mixing operations on pose keyframes already on the stack. Evaluation tasks can be used to create hierarchy as well, by creating a task that itself runs a program. This structure effectively allows our new Anim Mixer system to create an evaluation program from all the various sources of animation in a given frame, and then pass that program to whichever animation target needs to evaluate it, whether a custom Anim Instance, a slot in a blueprint, or an Anim Next injection point.

This first pass implements the main ECS systems that are used for this mixing, as well as an example Anim Target that uses a custom Anim Instance to evaluate the program and apply the resulting pose. So far I have only implemented the ECS system for processing animations created by the 'Animation' track (that is by evaluating an Anim Sequence at a given frame). Future checkins will add new animation track support (such as for control rig), as well as new animation targets, including crucial Anim Next support.

The ECS systems involved in this so far are:

* Mixed Skeletal Animation System- takes the place of the 'Skeletal Animation System' by ingesting entities created by animation sections and creating evaluation tasks for them to send on to the mixer. This is an example of what we expect other animation tracks to require in order to support this mixing workflow.
* Anim Mixer System- this is the system that ingests evaluation task components, builds evaluation programs for them, and sends them on to the appropriate animation target system.
* Anim Instance Target System- this is an example of an 'animation target' system. In this case it takes evaluation programs targeting it, and stores those evaluation programs in a custom anim instance proxy. It handles setting the anim instance on the actor to this custom anim instance as appropriate. Basically this is taking the place of the code in the Skeletal Animation System that did the same.
* Root Motion System- One of the other advantages of using the evaluation program is that you can evaluate a full program but only on specific elements- for example attributes or even just the trajectory/root motion. We take advantage of this in the Root Motion system to process just the root motion mixing ahead of the animation tick and then apply that as an additive transform on top of whatever other transform track operations may be taking place. In the future we may optimize this further once we add the ability to intersperse ECS tasks alongside other non-Sequencer tasks in the actor tick.

A few notes on other details:
* Evaluation Program has been changed to use TSharedPtr instead of TUniquePtr so that we can better manage passing these tasks around from our system
* To support this,  new 'MixedAnimationPriority' and 'MixedAnimationTarget' fields have been added to MovieSceneSkeletalAnimationSection.h. Similar will be needed on any other animation producing tracks. Priority is used here to help order different animation producing entities in the resulting evaluation program list. Future work will be needed on UX to make this easier- we have plans for a 'Mixer Track' that can hold older tracks as subtracks for example.
* A new tag- 'AnimMixerPoseProducer' has been added to the ECS built in component types. At the moment this is mainly used to prevent the Skeletal Animation System from operating when we have this plugin enabled.
* For root motion, I added a new attribute for the 'root bone transform'. Generally speaking the anim next stuff has been built to deal with root motion as a 'delta', but in Sequencer we must deal with it absolutely to ensure determinism when jumping around to different points in the Sequence.
* This new implementation doesn't yet support motion vector simulation or the 'mirroring' feature of the anim track.

Some other todos include:
* We should be able to optimize out tasks with zero weight from the mix.
* We don't yet take advantage of the 'source pose' as an underlying base pose that can be mixed as well, but we should be able to add this relatively straightforwardly.

diff --git a/Engine/Plugins/Experimental/AnimNextAnimGraph/Source/AnimNextAnimGraph/Private/EvaluationVM/EvaluationProgram.cpp b/Engine/Plugins/Experimental/AnimNextAnimGraph/Source/AnimNextAnimGraph/Private/EvaluationVM/EvaluationProgram.cpp
--- a/Engine/Plugins/Experimental/AnimNextAnimGraph/Source/AnimNextAnimGraph/Private/EvaluationVM/EvaluationProgram.cpp
+++ b/Engine/Plugins/Experimental/AnimNextAnimGraph/Source/AnimNextAnimGraph/Private/EvaluationVM/EvaluationProgram.cpp
@@ -22,7 +22,7 @@ namespace UE::AnimNext
 
 		if (VM.IsValid())
 		{
-			for (const TUniquePtr<FAnimNextEvaluationTask>& Task : Tasks)
+			for (const TSharedPtr<FAnimNextEvaluationTask>& Task : Tasks)
 			{
 				Task->Execute(VM);
 			}
@@ -34,7 +34,7 @@ namespace UE::AnimNext
 		FStringBuilderBase Result;
 
 		FString PropertyValue;
-		for (const TUniquePtr<FAnimNextEvaluationTask>& Task : Tasks)
+		for (const TSharedPtr<FAnimNextEvaluationTask>& Task : Tasks)
 		{
 			const UScriptStruct* TaskStruct = Task->GetStruct();
 
diff --git a/Engine/Plugins/Experimental/AnimNextAnimGraph/Source/AnimNextAnimGraph/Private/EvaluationVM/Tasks/ExecuteProgram.cpp b/Engine/Plugins/Experimental/AnimNextAnimGraph/Source/AnimNextAnimGraph/Private/EvaluationVM/Tasks/ExecuteProgram.cpp
--- a/Engine/Plugins/Experimental/AnimNextAnimGraph/Source/AnimNextAnimGraph/Private/EvaluationVM/Tasks/ExecuteProgram.cpp
+++ b/Engine/Plugins/Experimental/AnimNextAnimGraph/Source/AnimNextAnimGraph/Private/EvaluationVM/Tasks/ExecuteProgram.cpp
@@ -5,7 +5,7 @@
 #include "EvaluationVM/EvaluationProgram.h"
 #include "EvaluationVM/EvaluationVM.h"
 
-FAnimNextExecuteProgramTask FAnimNextExecuteProgramTask::Make(const UE::AnimNext::FEvaluationProgram* Program)
+FAnimNextExecuteProgramTask FAnimNextExecuteProgramTask::Make(TSharedPtr<const UE::AnimNext::FEvaluationProgram> Program)
 {
 	FAnimNextExecuteProgramTask Task;
 	Task.Program = Program;
@@ -14,7 +14,7 @@ FAnimNextExecuteProgramTask FAnimNextExecuteProgramTask::Make(const UE::AnimNext
 
 void FAnimNextExecuteProgramTask::Execute(UE::AnimNext::FEvaluationVM& VM) const
 {
-	if (Program != nullptr)
+	if (Program.IsValid())
 	{
 		Program->Execute(VM);
 	}
diff --git a/Engine/Plugins/Experimental/AnimNextAnimGraph/Source/AnimNextAnimGraph/Public/EvaluationVM/EvaluationProgram.h b/Engine/Plugins/Experimental/AnimNextAnimGraph/Source/AnimNextAnimGraph/Public/EvaluationVM/EvaluationProgram.h
--- a/Engine/Plugins/Experimental/AnimNextAnimGraph/Source/AnimNextAnimGraph/Public/EvaluationVM/EvaluationProgram.h
+++ b/Engine/Plugins/Experimental/AnimNextAnimGraph/Source/AnimNextAnimGraph/Public/EvaluationVM/EvaluationProgram.h
@@ -36,6 +36,10 @@ namespace UE::AnimNext
 		template<class TaskType>
 		void AppendTask(TaskType&& Task);
 
+		// Same as above, but uses TaskType copy constructor to move the task into the program.
+		template<class TaskType>
+		void AppendTaskPtr(const TSharedPtr<TaskType>& TaskPtr);
+
 		// Executes the current program on the provided virtual machine
 		void Execute(FEvaluationVM& VM) const;
 
@@ -48,7 +52,7 @@ namespace UE::AnimNext
 		FEvaluationProgram& operator=(const FEvaluationProgram&) = delete;
 
 		// List of tasks
-		TArray<TUniquePtr<FAnimNextEvaluationTask>> Tasks;
+		TArray<TSharedPtr<FAnimNextEvaluationTask>> Tasks;
 	};
 
 	//////////////////////////////////////////////////////////////////////////
@@ -60,7 +64,15 @@ namespace UE::AnimNext
 		static_assert(std::is_base_of<FAnimNextEvaluationTask, TaskType>::value, "Task type must derive from FAnimNextEvaluationTask");
 
 		// TODO: Use a fancy allocator to ensure all tasks are contiguous and tightly packed
-		TUniquePtr<FAnimNextEvaluationTask> TaskCopy = MakeUnique<TaskType>(MoveTemp(Task));
+		TSharedPtr<FAnimNextEvaluationTask> TaskCopy = MakeShared<TaskType>(MoveTemp(Task));
 		Tasks.Add(MoveTemp(TaskCopy));
 	}
+
+	template<class TaskType>
+	inline void FEvaluationProgram::AppendTaskPtr(const TSharedPtr<TaskType>& TaskPtr)
+	{
+		static_assert(std::is_base_of<FAnimNextEvaluationTask, TaskType>::value, "Task type must derive from FAnimNextEvaluationTask");
+
+		Tasks.Add(TaskPtr);
+	}
 }
diff --git a/Engine/Plugins/Experimental/AnimNextAnimGraph/Source/AnimNextAnimGraph/Public/EvaluationVM/Tasks/ExecuteProgram.h b/Engine/Plugins/Experimental/AnimNextAnimGraph/Source/AnimNextAnimGraph/Public/EvaluationVM/Tasks/ExecuteProgram.h
--- a/Engine/Plugins/Experimental/AnimNextAnimGraph/Source/AnimNextAnimGraph/Public/EvaluationVM/Tasks/ExecuteProgram.h
+++ b/Engine/Plugins/Experimental/AnimNextAnimGraph/Source/AnimNextAnimGraph/Public/EvaluationVM/Tasks/ExecuteProgram.h
@@ -25,11 +25,11 @@ struct ANIMNEXTANIMGRAPH_API FAnimNextExecuteProgramTask : public FAnimNextEvalu
 
 	DECLARE_ANIM_EVALUATION_TASK(FAnimNextExecuteProgramTask)
 
-	static FAnimNextExecuteProgramTask Make(const UE::AnimNext::FEvaluationProgram* Program);
+	static FAnimNextExecuteProgramTask Make(TSharedPtr<const UE::AnimNext::FEvaluationProgram> Program);
 
 	// Task entry point
 	virtual void Execute(UE::AnimNext::FEvaluationVM& VM) const override;
 
 	// The program to execute
-	const UE::AnimNext::FEvaluationProgram* Program = nullptr;
+	TSharedPtr<const UE::AnimNext::FEvaluationProgram> Program = nullptr;
 };
diff --git a/Engine/Plugins/Experimental/MovieSceneAnimMixer/MovieSceneAnimMixer.uplugin b/Engine/Plugins/Experimental/MovieSceneAnimMixer/MovieSceneAnimMixer.uplugin
new file mode 100644
--- /dev/null
+++ b/Engine/Plugins/Experimental/MovieSceneAnimMixer/MovieSceneAnimMixer.uplugin
@@ -0,0 +1,37 @@
+{
+	"FileVersion" : 1,
+	"Version" : 1,
+	"VersionName" : "0.1",
+	"FriendlyName" : "Sequencer Anim Mixer",
+	"Description" : "System for mixing layered animation in sequences",
+	"Category" : "Animation",
+	"CreatedBy" : "Epic Games, Inc.",
+	"CreatedByURL" : "https://epicgames.com",
+	"DocsURL" : "",
+	"MarketplaceURL" : "",
+	"SupportURL" : "",
+	"EnabledByDefault" : false,
+	"CanContainContent" : true,
+	"IsBetaVersion": false,
+	"IsExperimentalVersion": true,
+	"Installed" : false,
+	"Plugins":
+	[
+		{
+			"Name": "AnimNext",
+			"Enabled": true
+		},
+		{
+			"Name": "AnimNextAnimGraph",
+			"Enabled": true
+		}
+	],
+	"Modules" :
+	[
+		{
+			"Name" : "MovieSceneAnimMixer",
+			"Type" : "Runtime",
+			"LoadingPhase" : "Default"
+		}
+	]
+}
diff --git a/Engine/Plugins/Experimental/MovieSceneAnimMixer/Source/MovieSceneAnimMixer/MovieSceneAnimMixer.Build.cs b/Engine/Plugins/Experimental/MovieSceneAnimMixer/Source/MovieSceneAnimMixer/MovieSceneAnimMixer.Build.cs
new file mode 100644
--- /dev/null
+++ b/Engine/Plugins/Experimental/MovieSceneAnimMixer/Source/MovieSceneAnimMixer/MovieSceneAnimMixer.Build.cs
@@ -0,0 +1,31 @@
+// Copyright Epic Games, Inc. All Rights Reserved.
+
+namespace UnrealBuildTool.Rules
+{
+	public class MovieSceneAnimMixer : ModuleRules
+	{
+		public MovieSceneAnimMixer(ReadOnlyTargetRules Target) : base(Target)
+		{
+			PublicDependencyModuleNames.AddRange(
+				new string[]
+				{
+					"Core",
+					"CoreUObject",
+					"Engine",
+					"MovieScene",
+					"MovieSceneTracks",
+					"AnimGraphRuntime",
+					"AnimNext",
+					"AnimNextAnimGraph"
+				}
+			);
+
+			PrivateDependencyModuleNames.AddRange(
+				new string[]
+				{
+					"MovieSceneTracks",
+				}
+			);
+		}
+	}
+}
\ No newline at end of file
diff --git a/Engine/Plugins/Experimental/MovieSceneAnimMixer/Source/MovieSceneAnimMixer/Private/AnimMixerComponentTypes.cpp b/Engine/Plugins/Experimental/MovieSceneAnimMixer/Source/MovieSceneAnimMixer/Private/AnimMixerComponentTypes.cpp
new file mode 100644
--- /dev/null
+++ b/Engine/Plugins/Experimental/MovieSceneAnimMixer/Source/MovieSceneAnimMixer/Private/AnimMixerComponentTypes.cpp
@@ -0,0 +1,59 @@
+// Copyright Epic Games, Inc. All Rights Reserved.
+
+#include "AnimMixerComponentTypes.h"
+#include "EntitySystem/BuiltInComponentTypes.h"
+#include "EntitySystem/MovieSceneComponentRegistry.h"
+#include "EntitySystem/MovieSceneEntityFactoryTemplates.h"
+#include "EntitySystem/MovieSceneEntitySystemLinker.h"
+#include "EntitySystem/MovieScenePropertyComponentHandler.h"
+
+namespace UE::MovieScene
+{
+
+	const UE::Anim::FAttributeId FAnimMixerComponentTypes::RootTransformAttributeId(TEXT("RootTransform"), FCompactPoseBoneIndex(0));
+
+	static TUniquePtr<FAnimMixerComponentTypes> GAnimMixerComponentTypes;
+	static bool GAnimMixerComponentTypesDestroyed = false;
+
+	FAnimMixerComponentTypes* FAnimMixerComponentTypes::Get()
+	{
+		if (!GAnimMixerComponentTypes.IsValid())
+		{
+			check(!GAnimMixerComponentTypesDestroyed);
+			GAnimMixerComponentTypes.Reset(new FAnimMixerComponentTypes);
+		}
+		return GAnimMixerComponentTypes.Get();
+	}
+
+	void FAnimMixerComponentTypes::Destroy()
+	{
+		GAnimMixerComponentTypes.Reset();
+		GAnimMixerComponentTypesDestroyed = true;
+	}
+
+	FAnimMixerComponentTypes::FAnimMixerComponentTypes()
+	{
+		FBuiltInComponentTypes* BuiltInTypes = FBuiltInComponentTypes::Get();
+		FComponentRegistry* ComponentRegistry = UMovieSceneEntitySystemLinker::GetComponents();
+
+		ComponentRegistry->NewComponentType(&Priority, TEXT("Mixed Animation Priority"));
+		ComponentRegistry->NewComponentType(&Target, TEXT("Mixed Animation Target"));
+		ComponentRegistry->NewComponentType(&Task, TEXT("Mixed Animation Task"));
+		ComponentRegistry->NewComponentType(&MixerTask, TEXT("Mixed Animation Mixer Task"));
+		ComponentRegistry->NewComponentType(&RootMotionType, TEXT("Root Motion Type"));
+		ComponentRegistry->NewComponentType(&MeshComponent, TEXT("Mixed Animation Mesh Component"));
+
+		Tags.RequiresBlending = ComponentRegistry->NewTag(TEXT("Requires Blending"));
+
+		ComponentRegistry->Factories.DuplicateChildComponent(Priority);
+		ComponentRegistry->Factories.DuplicateChildComponent(Target);
+		ComponentRegistry->Factories.DuplicateChildComponent(Task);
+		ComponentRegistry->Factories.DuplicateChildComponent(MixerTask);
+		ComponentRegistry->Factories.DuplicateChildComponent(RootMotionType);
+		ComponentRegistry->Factories.DuplicateChildComponent(MeshComponent);
+
+		ComponentRegistry->Factories.DefineChildComponent(Tags.RequiresBlending, Tags.RequiresBlending);
+
+	}
+
+}
\ No newline at end of file
diff --git a/Engine/Plugins/Experimental/MovieSceneAnimMixer/Source/MovieSceneAnimMixer/Private/MovieSceneAnimMixerModule.cpp b/Engine/Plugins/Experimental/MovieSceneAnimMixer/Source/MovieSceneAnimMixer/Private/MovieSceneAnimMixerModule.cpp
new file mode 100644
--- /dev/null
+++ b/Engine/Plugins/Experimental/MovieSceneAnimMixer/Source/MovieSceneAnimMixer/Private/MovieSceneAnimMixerModule.cpp
@@ -0,0 +1,21 @@
+// Copyright Epic Games, Inc. All Rights Reserved.
+
+#include "MovieSceneAnimMixerModule.h"
+#include "Modules/ModuleInterface.h"
+#include "Modules/ModuleManager.h"
+
+namespace UE::MovieScene
+{
+
+void FMovieSceneAnimMixerModule::StartupModule()
+{
+}
+
+void FMovieSceneAnimMixerModule::ShutdownModule()
+{
+
+}
+
+}
+
+IMPLEMENT_MODULE(UE::MovieScene::FMovieSceneAnimMixerModule, MovieSceneAnimMixer)
diff --git a/Engine/Plugins/Experimental/MovieSceneAnimMixer/Source/MovieSceneAnimMixer/Private/MovieSceneAnimMixerModule.h b/Engine/Plugins/Experimental/MovieSceneAnimMixer/Source/MovieSceneAnimMixer/Private/MovieSceneAnimMixerModule.h
new file mode 100644
--- /dev/null
+++ b/Engine/Plugins/Experimental/MovieSceneAnimMixer/Source/MovieSceneAnimMixer/Private/MovieSceneAnimMixerModule.h
@@ -0,0 +1,18 @@
+﻿// Copyright Epic Games, Inc. All Rights Reserved.
+
+#pragma once
+
+#include "CoreMinimal.h"
+#include "Modules/ModuleInterface.h"
+
+namespace UE::MovieScene
+{
+
+class FMovieSceneAnimMixerModule : public IModuleInterface
+{
+	// IModuleInterface interface
+	virtual void StartupModule() override;
+	virtual void ShutdownModule() override;
+};
+
+}
diff --git a/Engine/Plugins/Experimental/MovieSceneAnimMixer/Source/MovieSceneAnimMixer/Private/Systems/MovieSceneAnimInstanceTargetSystem.cpp b/Engine/Plugins/Experimental/MovieSceneAnimMixer/Source/MovieSceneAnimMixer/Private/Systems/MovieSceneAnimInstanceTargetSystem.cpp
new file mode 100644
--- /dev/null
+++ b/Engine/Plugins/Experimental/MovieSceneAnimMixer/Source/MovieSceneAnimMixer/Private/Systems/MovieSceneAnimInstanceTargetSystem.cpp
@@ -0,0 +1,517 @@
+// Copyright Epic Games, Inc. All Rights Reserved.
+
+#include "Systems/MovieSceneAnimInstanceTargetSystem.h"
+
+#include "EntitySystem/BuiltInComponentTypes.h"
+#include "EntitySystem/Interrogation/MovieSceneInterrogationLinker.h"
+#include "EntitySystem/MovieSceneEntityIDs.h"
+#include "EntitySystem/MovieSceneEntitySystemLinker.h"
+#include "EntitySystem/MovieSceneEntitySystemRunner.h"
+#include "EntitySystem/MovieSceneEntitySystemTask.h"
+#include "MovieSceneTracksComponentTypes.h"
+
+#include "Evaluation/PreAnimatedState/IMovieScenePreAnimatedStorage.h"
+#include "Evaluation/PreAnimatedState/MovieScenePreAnimatedObjectStorage.h"
+#include "Evaluation/PreAnimatedState/MovieScenePreAnimatedStorageID.h"
+#include "Evaluation/PreAnimatedState/MovieScenePreAnimatedStorageID.inl"
+
+#include "EvaluationVM/EvaluationTask.h"
+#include "EvaluationVM/EvaluationProgram.h"
+#include "EvaluationVM/EvaluationVM.h"
+#include "GenerationTools.h"
+
+#include "AnimCustomInstanceHelper.h"
+#include "Animation/AnimSingleNodeInstance.h"
+#include "Components/SkeletalMeshComponent.h"
+#include "Engine/SkeletalMesh.h"
+#include "SkeletalMeshRestoreState.h"
+
+#include "Rendering/MotionVectorSimulation.h"
+#include "Systems/MovieSceneMotionVectorSimulationSystem.h"
+#include "Systems/MovieSceneComponentTransformSystem.h"
+#include "Systems/MovieSceneQuaternionInterpolationRotationSystem.h"
+#include "Systems/MovieSceneObjectPropertySystem.h"
+#include "EntitySystem/MovieScenePreAnimatedStateSystem.h"
+
+#include "Animation/AnimationPoseData.h"
+#include "BonePose.h"
+#include "Animation/BuiltInAttributeTypes.h"
+#include "Systems/MovieSceneAnimMixerSystem.h"
+#include "AnimMixerComponentTypes.h"
+
+#include "Graph/AnimNext_LODPose.h"
+
+#if WITH_EDITOR
+#include "Editor.h"
+#include "Editor/UnrealEdEngine.h"
+#include "UnrealEdGlobals.h"
+#endif
+
+#include UE_INLINE_GENERATED_CPP_BY_NAME(MovieSceneAnimInstanceTargetSystem)
+
+
+USequencerMixedAnimInstance::USequencerMixedAnimInstance(const FObjectInitializer& ObjectInitializer)
+	: Super(ObjectInitializer)
+{
+	bUseMultiThreadedAnimationUpdate = false;
+}
+
+FAnimInstanceProxy* USequencerMixedAnimInstance::CreateAnimInstanceProxy()
+{
+	return new FSequencerMixedAnimInstanceProxy(this);
+}
+
+UAnimInstance* USequencerMixedAnimInstance::GetSourceAnimInstance()
+{
+	return GetProxyOnGameThread<FSequencerMixedAnimInstanceProxy>().GetSourceAnimInstance();
+}
+
+/** Anim Instance Source info - created externally and used here */
+void USequencerMixedAnimInstance::SetSourceAnimInstance(UAnimInstance* SourceAnimInstance)
+{
+	USkeletalMeshComponent* MeshComponent = GetOwningComponent();
+	ensure(MeshComponent->GetAnimInstance() != SourceAnimInstance);
+
+	if (SourceAnimInstance)
+	{
+		// Add the current animation instance as a linked instance
+		FLinkedInstancesAdapter::AddLinkedInstance(MeshComponent, SourceAnimInstance);
+
+		// Direct the mixed anim instance to the current animation instance to evaluate as its source (input pose)
+		GetProxyOnGameThread<FSequencerMixedAnimInstanceProxy>().SetSourceAnimInstance(SourceAnimInstance, UAnimInstance::GetProxyOnGameThreadStatic<FAnimInstanceProxy>(SourceAnimInstance));
+	}
+	else
+	{
+		UAnimInstance* CurrentSourceAnimInstance = GetProxyOnGameThread<FSequencerMixedAnimInstanceProxy>().GetSourceAnimInstance();
+		// Remove the original instances from the linked instances as it should be reinstated as the main anim instance
+		FLinkedInstancesAdapter::RemoveLinkedInstance(MeshComponent, CurrentSourceAnimInstance);
+
+		// Null out the animation instance used to as the input source for the mixed anim instance
+		GetProxyOnGameThread<FSequencerMixedAnimInstanceProxy>().SetSourceAnimInstance(nullptr, nullptr);
+	}
+}
+
+void USequencerMixedAnimInstance::SetMixerTask(TSharedPtr<FAnimNextEvaluationTask> InEvalTask)
+{
+	GetProxyOnGameThread<FSequencerMixedAnimInstanceProxy>().SetMixerTask(InEvalTask);
+}
+
+void FSequencerMixedAnimInstanceProxy::Initialize(UAnimInstance* InAnimInstance)
+{
+	FAnimInstanceProxy::Initialize(InAnimInstance);
+	UpdateCounter.Reset();
+}
+
+void FSequencerMixedAnimInstanceProxy::CacheBones()
+{
+	if (bBoneCachesInvalidated)
+	{
+		if (CurrentSourceProxy)
+		{
+			FAnimationCacheBonesContext InputContext(CurrentSourceProxy);
+			if (SourcePose.GetLinkNode())
+			{
+				SourcePose.CacheBones(InputContext);
+			}
+			else
+			{
+				CacheBonesInputProxy(CurrentSourceProxy);
+			}
+		}
+
+		bBoneCachesInvalidated = false;
+	}
+}
+
+void FSequencerMixedAnimInstanceProxy::PreEvaluateAnimation(UAnimInstance* InAnimInstance)
+{
+	Super::PreEvaluateAnimation(InAnimInstance);
+
+	if (CurrentSourceAnimInstance)
+	{
+		CurrentSourceAnimInstance->PreEvaluateAnimation();
+	}
+}
+
+bool FSequencerMixedAnimInstanceProxy::Evaluate(FPoseContext& Output)
+{
+	bool bHasValidOutput = false;
+	if (CurrentSourceProxy)
+	{
+		FBoneContainer& SourceRequiredBones = CurrentSourceProxy->GetRequiredBones();
+		if (SourceRequiredBones.IsValid())
+		{
+			FPoseContext InnerOutput(CurrentSourceProxy, Output.ExpectsAdditivePose());
+
+			// if no linked node, just use Evaluate of proxy
+			if (FAnimNode_Base* InputNode = SourcePose.GetLinkNode())
+			{
+				CurrentSourceProxy->EvaluateAnimation_WithRoot(InnerOutput, InputNode);
+			}
+			else if (CurrentSourceProxy->HasRootNode())
+			{
+				CurrentSourceProxy->EvaluateAnimationNode(InnerOutput);
+			}
+			else
+			{
+				EvaluateInputProxy(CurrentSourceProxy, InnerOutput);
+			}
+
+			Output.Pose.MoveBonesFrom(InnerOutput.Pose);
+			Output.Curve.MoveFrom(InnerOutput.Curve);
+			Output.CustomAttributes.MoveFrom(InnerOutput.CustomAttributes);
+			bHasValidOutput = true;
+		}
+		else
+		{
+			Output.ResetToRefPose();
+		}
+	}
+	else
+	{
+		Output.ResetToRefPose();
+	}
+
+	if (MixerTask)
+	{
+		using namespace UE::AnimNext;
+		FDataHandle RefPoseHandle = FDataRegistry::Get()->GetOrGenerateReferencePose(SkeletalMeshComponent);
+		FAnimNextGraphReferencePose GraphReferencePose(RefPoseHandle);
+
+		const UE::AnimNext::FReferencePose& RefPose = RefPoseHandle.GetRef<UE::AnimNext::FReferencePose>();
+		FAnimNextGraphLODPose ResultPose;
+		ResultPose.LODPose = FLODPoseHeap(RefPose, LODLevel, true, Output.ExpectsAdditivePose());
+
+		{
+			FEvaluationVM EvaluationVM(EEvaluationFlags::All, RefPose, LODLevel);
+
+			// Use the output pose in the mixer as the 'base pose'- push this pose first.
+			FKeyframeState Keyframe = EvaluationVM.MakeUninitializedKeyframe(false);
+			FGenerationTools::RemapPose(Output, Keyframe.Pose);
+			Keyframe.Curves.CopyFrom(Output.Curve);
+			// TODO: There is not a remap attributes the other way- do we need one?
+			Keyframe.Attributes.CopyFrom(Output.CustomAttributes);
+
+			EvaluationVM.PushValue(KEYFRAME_STACK_NAME, MakeUnique<FKeyframeState>(MoveTemp(Keyframe)));
+
+			MixerTask->Execute(EvaluationVM);
+
+			TUniquePtr<FKeyframeState> EvaluatedKeyframe;
+			if (EvaluationVM.PopValue(KEYFRAME_STACK_NAME, EvaluatedKeyframe))
+			{
+				ResultPose.LODPose.CopyFrom(EvaluatedKeyframe->Pose);
+				ResultPose.Curves.CopyFrom(EvaluatedKeyframe->Curves);
+				ResultPose.Attributes.CopyFrom(EvaluatedKeyframe->Attributes);
+				bHasValidOutput = true;
+			}
+
+			if (!bHasValidOutput)
+			{
+				// We need to output a valid pose, generate one
+				FKeyframeState ReferenceKeyframe = EvaluationVM.MakeReferenceKeyframe(Output.ExpectsAdditivePose());
+				ResultPose.LODPose.CopyFrom(ReferenceKeyframe.Pose);
+				ResultPose.Curves.CopyFrom(ReferenceKeyframe.Curves);
+				ResultPose.Attributes.CopyFrom(ReferenceKeyframe.Attributes);
+			}
+		}
+
+		FGenerationTools::RemapPose(ResultPose.LODPose, Output);
+		Output.Curve.CopyFrom(ResultPose.Curves);
+		FGenerationTools::RemapAttributes(ResultPose.LODPose, ResultPose.Attributes, Output);
+	}
+
+
+	return true;
+}
+
+void FSequencerMixedAnimInstanceProxy::UpdateAnimationNode(const FAnimationUpdateContext& InContext)
+{
+	UpdateCounter.Increment();
+
+	if (CurrentSourceProxy)
+	{
+		FAnimationUpdateContext SourceContext = InContext.WithOtherProxy(CurrentSourceProxy);
+		if (FAnimNode_Base* SourceNode = SourcePose.GetLinkNode())
+		{
+			CurrentSourceProxy->UpdateAnimation_WithRoot(SourceContext, SourceNode, TEXT("AnimGraph"));
+		}
+		else if (CurrentSourceProxy->HasRootNode())
+		{
+			CurrentSourceProxy->UpdateAnimationNode(SourceContext);
+		}
+		else
+		{
+			UpdateInputProxy(CurrentSourceProxy, SourceContext);
+		}
+	}
+}
+
+void FSequencerMixedAnimInstanceProxy::LinkSourcePose(UAnimInstance* InSourceInstance, FAnimInstanceProxy* InSourceProxy)
+{
+	UnlinkSourcePose();
+
+	if (InSourceInstance)
+	{
+		CurrentSourceAnimInstance = InSourceInstance;
+		CurrentSourceProxy = InSourceProxy;
+		SourcePose.SetLinkNode(CurrentSourceProxy->GetRootNode());
+
+		// reset counter, so that input proxy can restart
+		ResetCounterInputProxy(CurrentSourceProxy);
+	}
+}
+
+void FSequencerMixedAnimInstanceProxy::UnlinkSourcePose()
+{
+	CurrentSourceProxy = nullptr;
+	CurrentSourceAnimInstance = nullptr;
+	SourcePose.SetLinkNode(nullptr);
+}
+
+
+void FSequencerMixedAnimInstanceProxy::SetSourceAnimInstance(UAnimInstance* SourceAnimInstance, FAnimInstanceProxy* SourceAnimInputProxy)
+{
+	UnlinkSourcePose();
+
+	if (SourceAnimInstance)
+	{
+		LinkSourcePose(SourceAnimInstance, SourceAnimInputProxy);
+	}
+}
+
+
+void FSequencerMixedAnimInstanceProxy::SetMixerTask(TSharedPtr<FAnimNextEvaluationTask> InEvalTask)
+{
+	MixerTask = InEvalTask;
+}
+
+namespace UE::MovieScene
+{
+	/** Helper function to get our sequencer animation instance from a skeletal mesh component */
+	USequencerMixedAnimInstance* GetAnimSequencerInstance(USkeletalMeshComponent* SkeletalMeshComponent)
+	{
+		return Cast<USequencerMixedAnimInstance>(SkeletalMeshComponent->GetAnimInstance());
+	}
+
+	/** ------------------------------------------------------------------------- */
+
+	/** Pre-animated state for skeletal animations */
+	struct FPreAnimatedMixedSkeletalAnimationState
+	{
+		EAnimationMode::Type AnimationMode;
+		TStrongObjectPtr<UAnimInstance> CachedAnimInstance;
+		FSkeletalMeshRestoreState SkeletalMeshRestoreState;
+	};
+
+	/** Pre-animation traits for skeletal animations */
+	struct FPreAnimatedAnimInstanceMixedSkeletalAnimationTraits : FBoundObjectPreAnimatedStateTraits
+	{
+		using KeyType = FObjectKey;
+		using StorageType = FPreAnimatedMixedSkeletalAnimationState;
+
+		FPreAnimatedMixedSkeletalAnimationState CachePreAnimatedValue(const KeyType& Object)
+		{
+			FPreAnimatedMixedSkeletalAnimationState OutCachedValue;
+			USkeletalMeshComponent* Component = Cast<USkeletalMeshComponent>(Object.ResolveObjectPtr());
+			if (ensure(Component))
+			{
+				OutCachedValue.AnimationMode = Component->GetAnimationMode();
+				OutCachedValue.CachedAnimInstance.Reset(Component->AnimScriptInstance);
+				OutCachedValue.SkeletalMeshRestoreState.SaveState(Component);
+			}
+			return OutCachedValue;
+		}
+
+		void RestorePreAnimatedValue(const KeyType& Object, StorageType& InOutCachedValue, const FRestoreStateParams& Params)
+		{
+			USkeletalMeshComponent* Component = Cast<USkeletalMeshComponent>(Object.ResolveObjectPtr());
+			if (!Component || !Component->IsRegistered())
+			{
+				return;
+			}
+
+			FAnimCustomInstanceHelper::UnbindFromSkeletalMeshComponent<USequencerMixedAnimInstance>(Component);
+
+			// Restore LOD before reinitializing anim instance
+			InOutCachedValue.SkeletalMeshRestoreState.RestoreLOD(Component);
+
+			if (Component->GetAnimationMode() != InOutCachedValue.AnimationMode)
+			{
+				// this SetAnimationMode reinitializes even if the mode is same
+				// if we're using same anim blueprint, we don't want to keep reinitializing it. 
+				Component->SetAnimationMode(InOutCachedValue.AnimationMode);
+			}
+			if (InOutCachedValue.CachedAnimInstance.Get())
+			{
+				Component->AnimScriptInstance = InOutCachedValue.CachedAnimInstance.Get();
+				InOutCachedValue.CachedAnimInstance.Reset();
+				if (Component->AnimScriptInstance && Component->GetSkeletalMeshAsset() && Component->AnimScriptInstance->CurrentSkeleton != Component->GetSkeletalMeshAsset()->GetSkeleton())
+				{
+					//the skeleton may have changed so need to recalc required bones as needed.
+					Component->AnimScriptInstance->CurrentSkeleton = Component->GetSkeletalMeshAsset()->GetSkeleton();
+					//Need at least RecalcRequiredbones and UpdateMorphTargetrs
+					Component->InitializeAnimScriptInstance(true);
+				}
+			}
+
+			// Restore pose after unbinding to force the restored pose
+			Component->SetUpdateAnimationInEditor(true);
+			Component->SetUpdateClothInEditor(true);
+			if (!Component->IsPostEvaluatingAnimation())
+			{
+				Component->TickAnimation(0.f, false);
+				Component->RefreshBoneTransforms();
+				Component->RefreshFollowerComponents();
+				Component->UpdateComponentToWorld();
+				Component->FinalizeBoneTransform();
+				Component->MarkRenderTransformDirty();
+				Component->MarkRenderDynamicDataDirty();
+			}
+
+			// Reset the mesh component update flag and animation mode to what they were before we animated the object
+			InOutCachedValue.SkeletalMeshRestoreState.RestoreState(Component);
+
+			// if not game world, don't clean this up
+			if (Component->GetWorld() != nullptr && Component->GetWorld()->IsGameWorld() == false)
+			{
+				Component->ClearMotionVector();
+			}
+		}
+	};
+
+	/** Pre-animation storage for skeletal animations */
+	struct FPreAnimatedAnimInstanceMixedSkeletalAnimationStorage : TPreAnimatedStateStorage_ObjectTraits<FPreAnimatedAnimInstanceMixedSkeletalAnimationTraits>
+	{
+		static TAutoRegisterPreAnimatedStorageID<FPreAnimatedAnimInstanceMixedSkeletalAnimationStorage> StorageID;
+	};
+
+	TAutoRegisterPreAnimatedStorageID<FPreAnimatedAnimInstanceMixedSkeletalAnimationStorage> FPreAnimatedAnimInstanceMixedSkeletalAnimationStorage::StorageID;
+
+	/** Task for evaluating and applying animation tasks */
+	struct FEvaluateAndApplyAnimationTasks
+	{
+	private:
+
+		UMovieSceneEntitySystemLinker* Linker;
+
+		TSharedPtr<FPreAnimatedAnimInstanceMixedSkeletalAnimationStorage> PreAnimatedStorage;
+
+	public:
+
+		FEvaluateAndApplyAnimationTasks(UMovieSceneEntitySystemLinker* InLinker)
+			: Linker(InLinker)
+		{
+			PreAnimatedStorage = InLinker->PreAnimatedState.GetOrCreateStorage<FPreAnimatedAnimInstanceMixedSkeletalAnimationStorage>();
+		}
+
+		void ForEachEntity
+		(
+			FMovieSceneEntityID EntityID,
+			FObjectComponent MeshComponent,
+			const TInstancedStruct<FMovieSceneMixedAnimationTarget>& Target,
+			TSharedPtr<FAnimNextEvaluationTask> MixerTask
+		) const
+		{
+			// Invalid or incompatible target, return.
+			if (!Target.GetPtr<FMovieSceneAnimInstanceTarget>())
+			{
+				return;
+			}
+
+			USkeletalMeshComponent* SkeletalMeshComponent = CastChecked<USkeletalMeshComponent>(MeshComponent.GetObject());
+			ensureMsgf(SkeletalMeshComponent, TEXT("Attempting to apply animation to an anim instance without a valid skeletal mesh component."));
+
+			if (!SkeletalMeshComponent || !SkeletalMeshComponent->GetSkeletalMeshAsset())
+			{
+				return;
+			}
+
+			// Cache pre-animated state for this bound object before doing anything.
+			// We don't yet track what entities have already started animated vs. entities that just started this frame,
+			// so we just process all the currently active ones. If they are already tracked and have already had their
+			// pre-animated state saved, it these calls will just early return.
+
+			// For now, we always restore state
+			PreAnimatedStorage->BeginTrackingEntity(EntityID, true, FRootInstanceHandle(), SkeletalMeshComponent);
+			
+			FCachePreAnimatedValueParams CacheParams;
+			PreAnimatedStorage->CachePreAnimatedValue(CacheParams, SkeletalMeshComponent);
+
+
+			// Setup custom anim instance, using the current anim instance as a 'source'.
+			bool bWasCreated = false;
+			
+			USequencerMixedAnimInstance* SequencerInstance = FAnimCustomInstanceHelper::BindToSkeletalMeshComponent<USequencerMixedAnimInstance>(SkeletalMeshComponent, bWasCreated);
+
+			SequencerInstance->SetMixerTask(MixerTask);
+
+			// TODO: Figure out- can we do motion vector sim here with the blended anim? Or do we need to do something different
+			/*if (InSkeletalAnimations.SimulatedAnimations.Num() != 0)
+			{
+				UMovieSceneMotionVectorSimulationSystem* MotionVectorSim = Linker->FindSystem<UMovieSceneMotionVectorSimulationSystem>();
+				if (MotionVectorSim && MotionVectorSim->IsSimulationEnabled())
+				{
+					ApplyAnimations(SkeletalMeshComponent, InSkeletalAnimations.SimulatedAnimations);
+					SkeletalMeshComponent->TickAnimation(0.f, false);
+					SkeletalMeshComponent->ForceMotionVector();
+
+					SimulateMotionVectors(SkeletalMeshComponent, MotionVectorSim);
+				}
+			}*/
+
+			if (!SkeletalMeshComponent->IsPostEvaluatingAnimation() && SkeletalMeshComponent->PoseTickedThisFrame())
+			{
+				SkeletalMeshComponent->TickAnimation(0.f, false);
+
+				SkeletalMeshComponent->RefreshBoneTransforms();
+				SkeletalMeshComponent->RefreshFollowerComponents();
+				SkeletalMeshComponent->UpdateComponentToWorld();
+				SkeletalMeshComponent->FinalizeBoneTransform();
+				SkeletalMeshComponent->MarkRenderTransformDirty();
+				SkeletalMeshComponent->MarkRenderDynamicDataDirty();
+			}
+		}
+	};
+
+} // namespace UE::MovieScene
+
+UMovieSceneAnimInstanceTargetSystem::UMovieSceneAnimInstanceTargetSystem(const FObjectInitializer& ObjInit)
+	: Super(ObjInit)
+{
+	using namespace UE::MovieScene;
+
+	FMovieSceneTracksComponentTypes* TrackComponents = FMovieSceneTracksComponentTypes::Get();
+	FAnimMixerComponentTypes* AnimMixerComponents = FAnimMixerComponentTypes::Get();
+
+	RelevantComponent = AnimMixerComponents->MixerTask;
+	Phase = ESystemPhase::Scheduling;
+
+	SystemCategories |= FSystemInterrogator::GetExcludedFromInterrogationCategory();
+
+	if (HasAnyFlags(RF_ClassDefaultObject))
+	{
+		DefineImplicitPrerequisite(UMovieSceneAnimMixerSystem::StaticClass(), GetClass());
+		DefineImplicitPrerequisite(GetClass(), UMovieSceneRestorePreAnimatedStateSystem::StaticClass());
+	}
+}
+
+void UMovieSceneAnimInstanceTargetSystem::OnSchedulePersistentTasks(UE::MovieScene::IEntitySystemScheduler* TaskScheduler)
+{
+	using namespace UE::MovieScene;
+
+	FBuiltInComponentTypes* BuiltInComponents = FBuiltInComponentTypes::Get();
+	FMovieSceneTracksComponentTypes* TrackComponents = FMovieSceneTracksComponentTypes::Get();
+	FAnimMixerComponentTypes* AnimMixerComponents = FAnimMixerComponentTypes::Get();
+
+	// Evaluate anim tasks targeting a custom anim instance.
+	FTaskParams Params(TEXT("Apply Animation Tasks"));
+	Params.ForceGameThread();
+	FTaskID EvaluateTask = FEntityTaskBuilder()
+		.ReadEntityIDs()
+		.Read(AnimMixerComponents->MeshComponent)
+		.Read(AnimMixerComponents->Target)
+		.Read(AnimMixerComponents->MixerTask)
+		.SetParams(Params)
+		.Schedule_PerEntity<FEvaluateAndApplyAnimationTasks>(&Linker->EntityManager, TaskScheduler,
+			Linker);
+}
\ No newline at end of file
diff --git a/Engine/Plugins/Experimental/MovieSceneAnimMixer/Source/MovieSceneAnimMixer/Private/Systems/MovieSceneAnimMixerSystem.cpp b/Engine/Plugins/Experimental/MovieSceneAnimMixer/Source/MovieSceneAnimMixer/Private/Systems/MovieSceneAnimMixerSystem.cpp
new file mode 100644
--- /dev/null
+++ b/Engine/Plugins/Experimental/MovieSceneAnimMixer/Source/MovieSceneAnimMixer/Private/Systems/MovieSceneAnimMixerSystem.cpp
@@ -0,0 +1,385 @@
+// Copyright Epic Games, Inc. All Rights Reserved.
+
+#include "Systems/MovieSceneAnimMixerSystem.h"
+
+#include "EntitySystem/BuiltInComponentTypes.h"
+#include "AnimMixerComponentTypes.h"
+#include "EntitySystem/MovieSceneEntityIDs.h"
+#include "EntitySystem/MovieSceneEntitySystemLinker.h"
+#include "EntitySystem/MovieSceneEntitySystemRunner.h"
+#include "EntitySystem/MovieSceneEntitySystemTask.h"
+#include "EntitySystem/MovieSceneBoundObjectInstantiator.h"
+#include "EntitySystem/MovieSceneBoundSceneComponentInstantiator.h"
+#include "EntitySystem/Interrogation/MovieSceneInterrogationLinker.h"
+#include "EntitySystem/MovieSceneEntityMutations.h"
+#include "MovieSceneTracksComponentTypes.h"
+
+#include "Systems/WeightAndEasingEvaluatorSystem.h"
+#include "Systems/MovieScenePiecewiseDoubleBlenderSystem.h"
+#include "EvaluationVM/Tasks/ExecuteProgram.h"
+#include "EvaluationVM/Tasks/BlendKeyframes.h"
+#include "EvaluationVM/Tasks/ApplyAdditiveKeyframe.h"
+#include "EvaluationVM/EvaluationProgram.h"
+#include "Algo/ForEach.h"
+#include "Algo/MaxElement.h"
+#include "Systems/MovieScenePropertyInstantiator.h"
+#include "EntitySystem/Interrogation/MovieSceneInterrogatedPropertyInstantiator.h"
+#include "UObject/Object.h"
+#include "EntitySystem/MovieSceneEntityGroupingSystem.h"
+#include "Components/SkeletalMeshComponent.h"
+#include "GameFramework/Actor.h"
+
+#include UE_INLINE_GENERATED_CPP_BY_NAME(MovieSceneAnimMixerSystem)
+
+UMovieSceneAnimMixerSystem::UMovieSceneAnimMixerSystem(const FObjectInitializer& ObjInit)
+	: Super(ObjInit)
+{
+	using namespace UE::MovieScene;
+
+	FAnimMixerComponentTypes* AnimMixerComponents = FAnimMixerComponentTypes::Get();
+	RelevantComponent = AnimMixerComponents->Task;
+	Phase = ESystemPhase::Instantiation | ESystemPhase::Scheduling;
+	SystemCategories = EEntitySystemCategory::BlenderSystems;
+
+	if (HasAnyFlags(RF_ClassDefaultObject))
+	{
+		FBuiltInComponentTypes* BuiltInComponents = FBuiltInComponentTypes::Get();
+
+		DefineImplicitPrerequisite(GetClass(), UMovieSceneEntityGroupingSystem::StaticClass());
+		DefineImplicitPrerequisite(GetClass(), UMovieScenePropertyInstantiatorSystem::StaticClass());
+		DefineImplicitPrerequisite(GetClass(), UMovieSceneInterrogatedPropertyInstantiatorSystem::StaticClass());
+		DefineImplicitPrerequisite(UMovieSceneGenericBoundObjectInstantiator::StaticClass(), GetClass());
+		DefineImplicitPrerequisite(UMovieSceneBoundSceneComponentInstantiator::StaticClass(), GetClass());
+		DefineComponentConsumer(GetClass(), BuiltInComponents->BoundObjectKey);
+		DefineComponentConsumer(GetClass(), BuiltInComponents->BoundObject);
+		DefineComponentConsumer(GetClass(), AnimMixerComponents->Task);
+		DefineComponentConsumer(GetClass(), AnimMixerComponents->Target);
+		DefineComponentConsumer(GetClass(), AnimMixerComponents->Priority);
+		DefineComponentProducer(GetClass(), AnimMixerComponents->MixerTask);
+	}
+}
+
+void UMovieSceneAnimMixerSystem::OnRun(FSystemTaskPrerequisites& InPrerequisites, FSystemSubsequentTasks& Subsequents)
+{
+	using namespace UE::MovieScene;
+
+	TSharedRef<FMovieSceneEntitySystemRunner> Runner = Linker->GetRunner();
+
+	if (Runner->GetCurrentPhase() == ESystemPhase::Instantiation)
+	{
+		FBuiltInComponentTypes* BuiltInComponents = FBuiltInComponentTypes::Get();
+		FMovieSceneTracksComponentTypes* TrackComponents = FMovieSceneTracksComponentTypes::Get();
+		FAnimMixerComponentTypes* AnimMixerComponents = FAnimMixerComponentTypes::Get();
+
+		// ----------------------------------------------------------------------------------------------------------------------------------------
+		// Step 1 - Remove expired mixer entries.
+
+		FEntityTaskBuilder()
+			.ReadEntityIDs()
+			.Read(BuiltInComponents->BoundObjectKey)
+			.Read(AnimMixerComponents->Target)
+			.FilterAll({ BuiltInComponents->Tags.NeedsUnlink, AnimMixerComponents->Task})
+			.Iterate_PerEntity(&Linker->EntityManager, [this](FMovieSceneEntityID EntityID, FObjectKey BoundObjectKey, const TInstancedStruct<FMovieSceneMixedAnimationTarget>& Target)
+				{
+					if (FMovieSceneAnimMixer* Mixer = Mixers.Find(FMovieSceneAnimMixerKey({ BoundObjectKey, Target })))
+					{
+						Mixer->Entries.RemoveAll([&EntityID](const FMovieSceneAnimMixerEntry& Entry) { return Entry.EntityID == EntityID;});
+					}
+				});
+
+		// ----------------------------------------------------------------------------------------------------------------------------------------
+		// Step 2 - Create new mixer entries for new anim tasks, gathering any new mixer entities that need to be created
+		TArray<FMovieSceneAnimMixerKey> NewMixers;
+
+		FEntityTaskBuilder()
+		.ReadEntityIDs()
+		.Read(BuiltInComponents->BoundObjectKey)
+		.Read(AnimMixerComponents->Target)
+		.Read(AnimMixerComponents->Task)
+		.Read(AnimMixerComponents->Priority)
+		.ReadOptional(AnimMixerComponents->RootMotionType)
+		.FilterAll({ BuiltInComponents->Tags.NeedsLink})
+		.Iterate_PerEntity( &Linker->EntityManager,
+		[this, &NewMixers, BuiltInComponents, AnimMixerComponents]
+		(
+			FMovieSceneEntityID EntityID,
+			FObjectKey BoundObjectKey,
+			const TInstancedStruct<FMovieSceneMixedAnimationTarget>& Target,
+			TSharedPtr<FAnimNextEvaluationTask> Task,
+			int32 Priority,
+			const ESwapRootBone* RootMotionType)
+			{
+				FMovieSceneAnimMixerKey Key({ BoundObjectKey, Target });
+				FMovieSceneAnimMixer* Mixer = Mixers.Find(Key);
+				if (!Mixer)
+				{
+					Mixer = &Mixers.Add({ Key, FMovieSceneAnimMixer()});
+					NewMixers.Add(Key);
+				}
+				if (Mixer)
+				{
+					checkSlow(!Mixer->Entries.ContainsByPredicate([EntityID](const FMovieSceneAnimMixerEntry& Entry) { return Entry.EntityID == EntityID;}));
+					
+					const FComponentMask& Type = Linker->EntityManager.GetEntityType(EntityID);
+					FMovieSceneAnimMixerEntry NewEntry;
+					NewEntry.EntityID = EntityID;
+					NewEntry.EvalTask = Task;
+					NewEntry.Priority = Priority;
+					NewEntry.PoseWeight = 0;
+					NewEntry.RootMotionType = RootMotionType ? *RootMotionType : ESwapRootBone::SwapRootBone_None;
+					NewEntry.bAdditive = Type.ContainsAny({BuiltInComponents->Tags.AdditiveBlend, BuiltInComponents->Tags.AdditiveFromBaseBlend});
+					NewEntry.bRequiresBlend = Type.Contains(AnimMixerComponents->Tags.RequiresBlending);
+					int32 EntryIndex = Algo::LowerBound(Mixer->Entries, NewEntry);
+					Mixer->Entries.Insert(NewEntry, EntryIndex);
+				}
+			});
+
+		// ----------------------------------------------------------------------------------------------------------------------------------------
+		// Step 3 - Create new mixer entities
+		for (FMovieSceneAnimMixerKey NewMixerKey : NewMixers)
+		{
+			FMovieSceneAnimMixer& NewMixer = Mixers.FindChecked(NewMixerKey);
+			NewMixer.MixerEntityID = FEntityBuilder()
+				.Add(AnimMixerComponents->MeshComponent, FObjectComponent::Weak(NewMixerKey.BoundObjectKey.ResolveObjectPtr()))
+				.Add(AnimMixerComponents->Target, NewMixerKey.Target)
+				.Add(AnimMixerComponents->MixerTask, nullptr)
+				.Add(AnimMixerComponents->RootMotionType, ESwapRootBone::SwapRootBone_None)
+				.AddTag(BuiltInComponents->Tags.RestoreState) // TODO: For now we always restore state on the mixer when it gets unlinked.
+				.CreateEntity(&Linker->EntityManager);
+		}
+
+
+		// ----------------------------------------------------------------------------------------------------------------------------------------
+		// Step 4 - Update mixer entities, and remove stale ones
+		TArray<FMovieSceneAnimMixerKey> StaleMixers;
+
+		// Component mask for adding root motion support to mixers if necessary
+		FComponentMask RootMotionMixerComponents;
+		RootMotionMixerComponents.SetAll({
+			BuiltInComponents->PropertyBinding,
+			BuiltInComponents->BlenderType,
+			TrackComponents->ComponentTransform.PropertyTag,
+			BuiltInComponents->BoundObject, // The target of the root motion, to be filled in by the root motion system
+			BuiltInComponents->BoundObjectKey,
+			BuiltInComponents->Tags.AdditiveBlend, // We treat root motion as additive since it will be computed in component space.
+			BuiltInComponents->DoubleResult[0],
+			BuiltInComponents->DoubleResult[1],
+			BuiltInComponents->DoubleResult[2],
+			BuiltInComponents->DoubleResult[3],
+			BuiltInComponents->DoubleResult[4],
+			BuiltInComponents->DoubleResult[5],
+			BuiltInComponents->DoubleResult[6],
+			BuiltInComponents->DoubleResult[7],
+			BuiltInComponents->DoubleResult[8],
+			});
+
+		for (const TPair<FMovieSceneAnimMixerKey, FMovieSceneAnimMixer>& MixerPair : Mixers)
+		{
+			ESwapRootBone RootMotionType = ESwapRootBone::SwapRootBone_None;
+			const FMovieSceneAnimMixerEntry* MaxEntry = Algo::MaxElementBy(MixerPair.Value.Entries, [](const FMovieSceneAnimMixerEntry& Entry) { return Entry.RootMotionType; });
+			if (MaxEntry)
+			{
+				RootMotionType = MaxEntry->RootMotionType;
+			}
+
+			bool bHasRootMotionComponent = Linker->EntityManager.GetEntityType(MixerPair.Value.MixerEntityID).Contains(TrackComponents->ComponentTransform.PropertyTag);
+
+			if (MixerPair.Value.Entries.IsEmpty())
+			{
+				StaleMixers.Add(MixerPair.Key);
+				Linker->EntityManager.AddComponent(MixerPair.Value.MixerEntityID, BuiltInComponents->Tags.NeedsUnlink, EEntityRecursion::Full);
+			}
+			else if (RootMotionType != ESwapRootBone::SwapRootBone_None) 
+			{
+				USkeletalMeshComponent* SkeletalMeshComponent = Cast<USkeletalMeshComponent>(MixerPair.Key.BoundObjectKey.ResolveObjectPtr());
+				if (!SkeletalMeshComponent)
+				{
+					continue;
+				}
+				UObject* RootMotionTarget = SkeletalMeshComponent;
+				if (RootMotionType == ESwapRootBone::SwapRootBone_Actor)
+				{
+					RootMotionTarget = SkeletalMeshComponent->GetOwner()->GetRootComponent();
+				}
+				// Update root motion type and add mixer components as necessary
+				FTypelessMutation Mutation;
+				Mutation.AddMask = RootMotionMixerComponents;
+				FEntityBuilder()
+					.Add(AnimMixerComponents->RootMotionType, RootMotionType)
+					.Add(BuiltInComponents->PropertyBinding, FMovieScenePropertyBinding(TEXT("Transform"), TEXT("Transform")))
+					.Add(BuiltInComponents->BlenderType, UMovieScenePiecewiseDoubleBlenderSystem::StaticClass())
+					.Add(BuiltInComponents->BoundObject, RootMotionTarget)
+					.Add(BuiltInComponents->BoundObjectKey, RootMotionTarget)
+					.AddTagConditional(BuiltInComponents->Tags.NeedsLink, !bHasRootMotionComponent)
+					.MutateExisting(&Linker->EntityManager, MixerPair.Value.MixerEntityID, Mutation);
+			}
+			else if (RootMotionType == ESwapRootBone::SwapRootBone_None && bHasRootMotionComponent)
+			{
+				// Remove the root motion mixer components, but add NeedsLink again to poke the transform system to remove this
+				FTypelessMutation Mutation;
+				Mutation.RemoveMask = RootMotionMixerComponents;
+				FEntityBuilder()
+					.AddTag(BuiltInComponents->Tags.NeedsLink)
+					.MutateExisting(&Linker->EntityManager, MixerPair.Value.MixerEntityID, Mutation);
+			}
+		}
+		Algo::ForEach(StaleMixers, [this](const FMovieSceneAnimMixerKey& MixerKey) { Mixers.Remove(MixerKey); });
+	}
+}
+
+namespace UE::MovieScene
+{
+	struct FUpdateTaskEntities
+	{
+		FUpdateTaskEntities(UMovieSceneAnimMixerSystem* InSystem)
+			:System(InSystem)
+
+		{
+
+		}
+		void ForEachEntity(FMovieSceneEntityID EntityID, 
+							FObjectKey BoundObjectKey,
+							const TInstancedStruct<FMovieSceneMixedAnimationTarget>& Target,
+							TSharedPtr<FAnimNextEvaluationTask> Task,
+							const double* WeightAndEasings) const
+		{
+			FMovieSceneAnimMixerKey Key({ BoundObjectKey, Target });
+			if (FMovieSceneAnimMixer* Mixer = System->Mixers.Find(Key))
+			{
+				// Find the entry in question
+				if (FMovieSceneAnimMixerEntry* Entry = Mixer->Entries.FindByPredicate([EntityID](const FMovieSceneAnimMixerEntry& TestEntry) { return TestEntry.EntityID == EntityID; }))
+				{
+					Entry->EvalTask = Task;
+					Entry->PoseWeight = WeightAndEasings ? *WeightAndEasings : 1.0;
+				}
+			}
+		}
+
+		TObjectPtr<UMovieSceneAnimMixerSystem> System;
+	};
+
+	struct FEvaluateAnimMixers
+	{
+		FEvaluateAnimMixers(UMovieSceneAnimMixerSystem* InSystem)
+			:System(InSystem)
+		{
+
+		}
+
+		void ForEachEntity(FMovieSceneEntityID EntityID,
+			FObjectComponent MeshComponent,
+			const TInstancedStruct<FMovieSceneMixedAnimationTarget>& Target,
+			TSharedPtr<FAnimNextEvaluationTask>& MixerTask) const
+		{
+			if (FMovieSceneAnimMixer* Mixer = System->Mixers.Find(FMovieSceneAnimMixerKey({ MeshComponent.GetObject(), Target })))
+			{
+				// Create a new eval program
+				Mixer->EvaluationProgram = MakeShared<UE::AnimNext::FEvaluationProgram>();
+
+				if (!MixerTask.IsValid())
+				{
+					MixerTask = MakeShared<FAnimNextExecuteProgramTask>();
+				}
+				TSharedPtr<FAnimNextExecuteProgramTask> EvalProgramTask = StaticCastSharedPtr<FAnimNextExecuteProgramTask>(MixerTask);
+				EvalProgramTask->Program = Mixer->EvaluationProgram;
+
+				// From lowest to highest priority, add the tasks to the program.
+				// Each task either
+				// a) Is self-contained, adding a single keyframe to the stack, in which case it will require a blend accumulation task afterwards
+				// b) Relies on the already accumulated pose and outputs a fully mixed pose to the stack, in which case no additional blend accumulation task is required
+				// Note that we assume here that the target will push some kind of 'base pose' to the VM before evaluating the task. This allows us to blend in/out from gameplay for example.
+				for (const FMovieSceneAnimMixerEntry& MixerEntry : Mixer->Entries)
+				{
+					if (MixerEntry.EvalTask.IsValid())
+					{
+						Mixer->EvaluationProgram->AppendTaskPtr(MixerEntry.EvalTask);
+						if (MixerEntry.bRequiresBlend)
+						{
+							if (MixerEntry.bAdditive)
+							{
+								// Add additive blend task.
+								Mixer->EvaluationProgram->AppendTask(FAnimNextApplyAdditiveKeyframeTask::Make(MixerEntry.PoseWeight));
+							}
+							else
+							{
+								// Add absolute blend task
+								Mixer->EvaluationProgram->AppendTask(FAnimNextBlendTwoKeyframesTask::Make(MixerEntry.PoseWeight));
+							}
+						}
+					}
+				}
+			}
+		}
+
+		TObjectPtr<UMovieSceneAnimMixerSystem> System;
+	};
+}
+
+void UMovieSceneAnimMixerSystem::OnSchedulePersistentTasks(UE::MovieScene::IEntitySystemScheduler* TaskScheduler)
+{
+	using namespace UE::MovieScene;
+
+	FBuiltInComponentTypes* BuiltInComponents = FBuiltInComponentTypes::Get();
+	FMovieSceneTracksComponentTypes* TrackComponents = FMovieSceneTracksComponentTypes::Get();
+	FAnimMixerComponentTypes* AnimMixerComponents = FAnimMixerComponentTypes::Get();
+
+	// TODO: We should be able to optimize here if we figure out which tasks won't contribute to the final pose. This should be doable 
+	// with a little extra API and pre-examination of per-bone blend weights, etc.
+
+	// Update mixer entry tasks- todo this is slow with the current data hierarchy
+	// Maybe we want a flat map of entity id to entry, and then put indices into that into the mixer or something similar.
+	FTaskID UpdateTask = FEntityTaskBuilder()
+		.ReadEntityIDs()
+		.Read(BuiltInComponents->BoundObjectKey)
+		.Read(AnimMixerComponents->Target)
+		.Read(AnimMixerComponents->Task)
+		.ReadOptional(BuiltInComponents->WeightAndEasingResult)
+		.FilterNone({BuiltInComponents->Tags.NeedsUnlink})
+		.Schedule_PerEntity<FUpdateTaskEntities>(&Linker->EntityManager, TaskScheduler, this );
+
+
+	// For each mixer, build the evaluation program task list.
+	FTaskID MixTask = FEntityTaskBuilder()
+		.ReadEntityIDs()
+		.Read(AnimMixerComponents->MeshComponent)
+		.Read(AnimMixerComponents->Target)
+		.Write(AnimMixerComponents->MixerTask)
+		.FilterAll({ AnimMixerComponents->MixerTask })
+		.Schedule_PerEntity<FEvaluateAnimMixers>(&Linker->EntityManager, TaskScheduler, this);
+
+	TaskScheduler->AddPrerequisite(UpdateTask, MixTask);
+}
+
+void UMovieSceneAnimMixerSystem::OnCleanTaggedGarbage()
+{
+	using namespace UE::MovieScene;
+
+	FBuiltInComponentTypes* BuiltInComponents = FBuiltInComponentTypes::Get();
+	FAnimMixerComponentTypes* AnimMixerComponents = FAnimMixerComponentTypes::Get();
+
+	FEntityTaskBuilder()
+		.ReadEntityIDs()
+		.Read(BuiltInComponents->BoundObjectKey)
+		.Read(AnimMixerComponents->Target)
+		.FilterAll({ BuiltInComponents->Tags.NeedsUnlink, AnimMixerComponents->Task })
+		.Iterate_PerEntity(&Linker->EntityManager, [this](FMovieSceneEntityID EntityID, FObjectKey BoundObjectKey, const TInstancedStruct<FMovieSceneMixedAnimationTarget>& Target)
+			{
+				if (FMovieSceneAnimMixer* Mixer = Mixers.Find(FMovieSceneAnimMixerKey({ BoundObjectKey, Target })))
+				{
+					Mixer->Entries.RemoveAll([&EntityID](const FMovieSceneAnimMixerEntry& Entry) { return Entry.EntityID == EntityID; });
+				}
+			});
+
+	FEntityTaskBuilder()
+		.ReadEntityIDs()
+		.Read(AnimMixerComponents->MeshComponent)
+		.Read(AnimMixerComponents->Target)
+		.FilterAll({ BuiltInComponents->Tags.NeedsUnlink, AnimMixerComponents->MixerTask })
+		.Iterate_PerEntity(&Linker->EntityManager, [this]
+			(FMovieSceneEntityID EntityID,
+			FObjectComponent MeshComponent,
+			const TInstancedStruct<FMovieSceneMixedAnimationTarget>& Target)
+			{
+				Mixers.Remove(FMovieSceneAnimMixerKey({ MeshComponent.GetObject(), Target }));
+			});
+}
\ No newline at end of file
diff --git a/Engine/Plugins/Experimental/MovieSceneAnimMixer/Source/MovieSceneAnimMixer/Private/Systems/MovieSceneMixedSkeletalAnimationSystem.cpp b/Engine/Plugins/Experimental/MovieSceneAnimMixer/Source/MovieSceneAnimMixer/Private/Systems/MovieSceneMixedSkeletalAnimationSystem.cpp
new file mode 100644
--- /dev/null
+++ b/Engine/Plugins/Experimental/MovieSceneAnimMixer/Source/MovieSceneAnimMixer/Private/Systems/MovieSceneMixedSkeletalAnimationSystem.cpp
@@ -0,0 +1,411 @@
+// Copyright Epic Games, Inc. All Rights Reserved.
+
+#include "Systems/MovieSceneMixedSkeletalAnimationSystem.h"
+
+#include "Async/TaskGraphInterfaces.h"
+#include "EntitySystem/BuiltInComponentTypes.h"
+#include "AnimMixerComponentTypes.h"
+#include "EntitySystem/MovieSceneEntityIDs.h"
+#include "EntitySystem/MovieSceneEntitySystemLinker.h"
+#include "EntitySystem/MovieSceneEntitySystemRunner.h"
+#include "EntitySystem/MovieSceneEntitySystemTask.h"
+#include "MovieSceneTracksComponentTypes.h"
+#include "Sections/MovieSceneSkeletalAnimationSection.h"
+
+#include "Components/SkeletalMeshComponent.h"
+#include "Engine/SkeletalMesh.h"
+#include "SkeletalMeshRestoreState.h"
+
+#include "Systems/WeightAndEasingEvaluatorSystem.h"
+#include "Systems/MovieSceneObjectPropertySystem.h"
+#include "EntitySystem/MovieScenePreAnimatedStateSystem.h"
+
+#include "Components/SkeletalMeshComponent.h"
+#include "Component/AnimNextComponent.h"
+#include "Systems/MovieSceneComponentTransformSystem.h"
+#include "Systems/MovieSceneQuaternionInterpolationRotationSystem.h"
+#include "EntitySystem/MovieSceneBoundObjectInstantiator.h"
+#include "EntitySystem/MovieSceneBoundSceneComponentInstantiator.h"
+#include "EntitySystem/MovieSceneEntityMutations.h"
+#include "EntitySystem/Interrogation/MovieSceneInterrogationLinker.h"
+#include "AnimSequencerInstanceProxy.h"
+#include "Animation/AnimSequence.h"
+#include "Animation/AnimRootMotionProvider.h"
+#include "BonePose.h"
+#include "DecompressionTools.h"
+#include "EvaluationVM/EvaluationVM.h"
+#include "EvaluationVM/KeyframeState.h"
+#include "Animation/BuiltInAttributeTypes.h"
+#include "EntitySystem/MovieSceneEvalTimeSystem.h"
+
+#include UE_INLINE_GENERATED_CPP_BY_NAME(MovieSceneMixedSkeletalAnimationSystem)
+
+
+namespace UE::MovieScene
+{
+
+// Adds required anim mixer components onto skeletal animation section entities
+struct FSkeletalAnimMixerMutation : IMovieSceneEntityMutation
+{
+	FBuiltInComponentTypes* BuiltInComponents;
+	FMovieSceneTracksComponentTypes* TrackComponents; 
+	FAnimMixerComponentTypes* AnimMixerComponents;
+	FSkeletalAnimMixerMutation()
+	{
+		BuiltInComponents = FBuiltInComponentTypes::Get();
+		TrackComponents = FMovieSceneTracksComponentTypes::Get();
+		AnimMixerComponents = FAnimMixerComponentTypes::Get();
+	}
+
+	virtual void CreateMutation(FEntityManager* EntityManager, FComponentMask* InOutEntityComponentTypes) const
+	{
+		InOutEntityComponentTypes->SetAll(
+		{
+			TrackComponents->Tags.AnimMixerPoseProducer,
+			AnimMixerComponents->Priority,
+			AnimMixerComponents->Target,
+			AnimMixerComponents->Task,
+			AnimMixerComponents->RootMotionType,
+			AnimMixerComponents->Tags.RequiresBlending
+		});
+	}
+
+	virtual void InitializeAllocation(FEntityAllocation* Allocation, const FComponentMask& AllocationType) const
+	{
+		const int32 Num = Allocation->Num();
+		FEntityAllocationWriteContext WriteContext = FEntityAllocationWriteContext::NewAllocation();
+		TComponentReader<FMovieSceneSkeletalAnimationComponentData> SkeletalAnimationData = Allocation->ReadComponents(TrackComponents->SkeletalAnimation);
+		TComponentWriter<TInstancedStruct<FMovieSceneMixedAnimationTarget>> OutAnimTargets = Allocation->WriteComponents(AnimMixerComponents->Target, WriteContext);
+		TComponentWriter<TSharedPtr<FAnimNextEvaluationTask>> OutTasks = Allocation->WriteComponents(AnimMixerComponents->Task, WriteContext);
+		TComponentWriter<int32> OutPriorities = Allocation->WriteComponents(AnimMixerComponents->Priority, WriteContext);
+		TComponentWriter<ESwapRootBone> OutRootMotionTypes = Allocation->WriteComponents(AnimMixerComponents->RootMotionType, WriteContext);
+		for (int32 Index = 0; Index < Num; ++Index)
+		{
+			// Initialize a null (empty) task, will be written during evaluation
+			OutTasks[Index] = nullptr;
+
+			const FMovieSceneSkeletalAnimationComponentData& SkeletalAnim = SkeletalAnimationData[Index];
+			if (SkeletalAnim.Section)
+			{
+				if (SkeletalAnim.Section->MixedAnimationTarget.IsValid())
+				{
+					// Set the assigned target
+					OutAnimTargets[Index] = SkeletalAnim.Section->MixedAnimationTarget;
+				}
+				else
+				{
+					// Set a default target
+					OutAnimTargets[Index] = TInstancedStruct<FMovieSceneMixedAnimationTarget>::Make();
+				}
+
+				OutPriorities[Index] = SkeletalAnim.Section->MixedAnimationPriority;
+				OutRootMotionTypes[Index] = SkeletalAnim.Section->Params.SwapRootBone;
+			}
+		}
+	}
+};
+
+/** ------------------------------------------------------------------------- */
+/** Task for gathering active skeletal animations and setting up their tasks. */
+struct FGatherMixableSkeletalAnimations
+{
+	const FInstanceRegistry* InstanceRegistry;
+
+	FGatherMixableSkeletalAnimations(const FInstanceRegistry* InInstanceRegistry)
+		: InstanceRegistry(InInstanceRegistry)
+	{
+	}
+
+	void ForEachAllocation(
+		const FEntityAllocationProxy AllocationProxy,
+		TRead<FMovieSceneEntityID> EntityIDs,
+		TRead<FInstanceHandle> InstanceHandles,
+		TRead<UObject*> BoundObjects,
+		TRead<FMovieSceneSkeletalAnimationComponentData> SkeletalAnimations,
+		TReadOptional<FFrameTime> OptionalEvalTimes,
+		TWrite<TSharedPtr<FAnimNextEvaluationTask>> EvalTask) const
+	{
+		// Gather all the skeletal animations currently active in all sequences.
+		// We map these animations to their bound object, which means we might blend animations from different sequences
+		// that have bound to the same object.
+		const FBuiltInComponentTypes* BuiltInComponents = FBuiltInComponentTypes::Get();
+		const FEntityAllocation* Allocation = AllocationProxy.GetAllocation();
+		const int32 Num = Allocation->Num();
+		for (int32 Index = 0; Index < Num; ++Index)
+		{
+			FMovieSceneEntityID EntityID(EntityIDs[Index]);
+			const FInstanceHandle& InstanceHandle(InstanceHandles[Index]);
+			UObject* BoundObject(BoundObjects[Index]);
+			const FMovieSceneSkeletalAnimationComponentData& SkeletalAnimation(SkeletalAnimations[Index]);
+
+			// Get the full context, so we can get both the current and previous evaluation times.
+			const FSequenceInstance& SequenceInstance = InstanceRegistry->GetInstance(InstanceHandle);
+			const FMovieSceneContext& Context = SequenceInstance.GetContext();
+
+			const FFrameTime EvalFrameTime = OptionalEvalTimes ? OptionalEvalTimes[Index] : Context.GetTime();
+
+			// Calculate the time at which to evaluate the animation
+			const UMovieSceneSkeletalAnimationSection* AnimSection = SkeletalAnimation.Section;
+			const FMovieSceneSkeletalAnimationParams& AnimParams = AnimSection->Params;
+
+			// Get the bound skeletal mesh component.
+			USkeletalMeshComponent* SkeletalMeshComponent = CastChecked<USkeletalMeshComponent>(BoundObject);
+			if (!SkeletalMeshComponent || AnimParams.Animation == nullptr)
+			{
+				continue;
+			}
+			const float EvalTime = AnimParams.MapTimeToAnimation(AnimSection, EvalFrameTime, Context.GetFrameRate());
+			const float PreviousEvalTime = AnimParams.MapTimeToAnimation(AnimSection, Context.GetPreviousTime(), Context.GetFrameRate());
+
+			const EMovieScenePlayerStatus::Type PlayerStatus = Context.GetStatus();
+
+			const bool bPreviewPlayback = ShouldUsePreviewPlayback(PlayerStatus, *BoundObject);
+
+			// If the playback status is jumping, ie. one such occurrence is setting the time for thumbnail generation, disable anim notifies updates because it could fire audio.
+			// If the playback status is scrubbing, we disable notifies for now because we can't properly fire them in all cases until we get evaluation range info.
+			// We now layer this with the passed in notify toggle to force a disable in this case.
+			const bool bFireNotifies = !bPreviewPlayback || (PlayerStatus != EMovieScenePlayerStatus::Jumping && PlayerStatus != EMovieScenePlayerStatus::Stopped && PlayerStatus != EMovieScenePlayerStatus::Scrubbing);
+			const bool bPlaying = PlayerStatus == EMovieScenePlayerStatus::Playing;
+
+			// Don't fire notifies if looping around.
+			bool bLooped = false;
+			if (AnimParams.bReverse)
+			{
+				if (PreviousEvalTime <= EvalTime)
+				{
+					bLooped = true;
+				}
+			}
+			else if (PreviousEvalTime >= EvalTime)
+			{
+				bLooped = true;
+			}
+
+			TOptional<FSkeletalAnimationRootMotionOverride> RootMotionOverride;
+			UMovieSceneSkeletalAnimationSection::FRootMotionParams RootMotionParams;
+
+			AnimSection->GetRootMotion(EvalFrameTime, RootMotionParams);
+			if (RootMotionParams.Transform.IsSet())
+			{
+				RootMotionOverride = FSkeletalAnimationRootMotionOverride();
+				RootMotionOverride.GetValue().RootMotion = RootMotionParams.Transform.GetValue();
+				RootMotionOverride.GetValue().bBlendFirstChildOfRoot = RootMotionParams.bBlendFirstChildOfRoot;
+				RootMotionOverride.GetValue().ChildBoneIndex = RootMotionParams.ChildBoneIndex;
+				RootMotionOverride.GetValue().PreviousTransform = RootMotionParams.PreviousTransform.GetValue();
+			}
+
+			FMixedAnimSkeletalAnimationData AnimData;
+			AnimData.AnimSequence = Cast<UAnimSequence>(AnimParams.Animation);
+			AnimData.RootMotionOverride = RootMotionOverride;
+			AnimData.FromPosition = PreviousEvalTime;
+			AnimData.ToPosition = EvalTime;
+			AnimData.bFireNotifies = bFireNotifies && !AnimParams.bSkipAnimNotifiers && !bLooped;
+			AnimData.SwapRootBone = AnimParams.SwapRootBone;
+			// TODO: We want to allow for additive to be set on the section properties
+			AnimData.bAdditive = AnimData.AnimSequence->IsValidAdditive();
+
+			// TODO: Need to create a mirroring task and add that if appropriate
+			// TODO: Notifies- perhaps this should be a part of the animnext side of things?
+
+			if (!EvalTask[Index].IsValid())
+			{ 
+				// Create task if not yet created
+				EvalTask[Index] = MakeShared<FMovieSceneSkeletalAnimationEvaluationTask>();
+			}
+			TSharedPtr<FMovieSceneSkeletalAnimationEvaluationTask> AnimTask = StaticCastSharedPtr<FMovieSceneSkeletalAnimationEvaluationTask>(EvalTask[Index]);
+			AnimTask->AnimationData = AnimData;
+		}
+	}
+
+private:
+
+	static bool ShouldUsePreviewPlayback(EMovieScenePlayerStatus::Type PlayerStatus, UObject& RuntimeObject)
+	{
+		// We also use PreviewSetAnimPosition in PIE when not playing, as we can preview in PIE.
+		bool bIsNotInPIEOrNotPlaying = (RuntimeObject.GetWorld() && !RuntimeObject.GetWorld()->HasBegunPlay()) || PlayerStatus != EMovieScenePlayerStatus::Playing;
+		return GIsEditor && bIsNotInPIEOrNotPlaying;
+	}
+};
+	
+} // namespace UE::MovieScene
+
+void FMovieSceneSkeletalAnimationEvaluationTask::Execute(UE::AnimNext::FEvaluationVM& VM) const
+{
+	using namespace UE::AnimNext;
+	using namespace UE::MovieScene;
+	FAnimMixerComponentTypes* AnimMixerComponents = FAnimMixerComponentTypes::Get();
+	if (const UAnimSequence* AnimSequencePtr = AnimationData.AnimSequence.Get())
+	{
+		const bool bIsAdditive = AnimSequencePtr->IsValidAdditive();
+		FDeltaTimeRecord DeltaTime;
+		DeltaTime.Set(AnimationData.FromPosition, AnimationData.ToPosition);
+
+		const bool bExtractRootMotion = AnimationData.RootMotionOverride.IsSet();
+		const FAnimExtractContext ExtractionContext(AnimationData.ToPosition, bExtractRootMotion, DeltaTime, false);
+
+		FKeyframeState Keyframe = VM.MakeUninitializedKeyframe(bIsAdditive);
+
+		const bool bUseRawData = FDecompressionTools::ShouldUseRawData(AnimSequencePtr, Keyframe.Pose);
+
+
+		bool bGrabbedRootTransform = false;
+
+		TOptional<FTransform> RootBoneTransform;
+
+		if (EnumHasAnyFlags(VM.GetFlags(), EEvaluationFlags::Bones))
+		{
+			FDecompressionTools::GetAnimationPose(AnimSequencePtr, ExtractionContext, Keyframe.Pose, bUseRawData);
+
+			if (AnimationData.RootMotionOverride.IsSet())
+			{
+				int32 PoseIndex = INDEX_NONE;
+				if (!AnimationData.RootMotionOverride->bBlendFirstChildOfRoot)
+				{
+					PoseIndex = Keyframe.Pose.GetRefPose().GetLODBoneIndexFromSkeletonBoneIndex(0);
+				}
+				else if (AnimationData.RootMotionOverride->ChildBoneIndex != INDEX_NONE)
+				{
+					PoseIndex = Keyframe.Pose.GetRefPose().GetLODBoneIndexFromSkeletonBoneIndex(AnimationData.RootMotionOverride->ChildBoneIndex);
+				}					
+				if (PoseIndex != INDEX_NONE)
+				{
+					Keyframe.Pose.LocalTransformsView[PoseIndex] = AnimationData.RootMotionOverride->RootMotion * AnimationData.RootMotionOverride->PreviousTransform;
+				}
+			}
+
+			// Extract the root motion if we've chosen that option
+			if (AnimationData.SwapRootBone != ESwapRootBone::SwapRootBone_None)
+			{
+				int32 RootIndex = Keyframe.Pose.GetRefPose().GetLODBoneIndexFromSkeletonBoneIndex(0);
+				RootBoneTransform = Keyframe.Pose.LocalTransformsView[RootIndex];
+				Keyframe.Pose.LocalTransformsView[RootIndex] = FTransform::Identity;
+			}
+		}
+
+		if (EnumHasAnyFlags(VM.GetFlags(), EEvaluationFlags::Curves))
+		{
+			FDecompressionTools::GetAnimationCurves(AnimSequencePtr, ExtractionContext, Keyframe.Curves, bUseRawData);
+		}
+
+		if (EnumHasAnyFlags(VM.GetFlags(), EEvaluationFlags::Attributes))
+		{
+			FDecompressionTools::GetAnimationAttributes(AnimSequencePtr, ExtractionContext, Keyframe.Pose.GetRefPose(), Keyframe.Attributes, bUseRawData);
+		}
+
+		// Trajectory is currently held as an attribute
+		if (EnumHasAnyFlags(VM.GetFlags(), EEvaluationFlags::Attributes | EEvaluationFlags::Trajectory))
+		{
+			// If the sequence has root motion enabled, allow sampling of a root motion delta into the custom attribute container of the outgoing pose
+			if (AnimSequencePtr->HasRootMotion())
+			{
+				// TODO: We should cache the provider in the VM
+				// We have to grab two locks to get it and it won't change during graph evaluation
+				if (const UE::Anim::IAnimRootMotionProvider* RootMotionProvider = UE::Anim::IAnimRootMotionProvider::Get())
+				{
+					RootMotionProvider->SampleRootMotion(ExtractionContext.DeltaTimeRecord, *AnimSequencePtr, ExtractionContext.bLooping, Keyframe.Attributes);
+					// TODO: Do I need to override this too?
+				}
+			}
+
+			// If we haven't already grabbed the root transform, grab it here in case we're not doing the full bone sample
+			if (AnimationData.SwapRootBone != ESwapRootBone::SwapRootBone_None && !RootBoneTransform.IsSet())
+			{
+				if (AnimationData.RootMotionOverride.IsSet())
+				{
+					RootBoneTransform = AnimationData.RootMotionOverride->RootMotion * AnimationData.RootMotionOverride->PreviousTransform;
+				}
+				else
+				{
+					RootBoneTransform = AnimSequencePtr->ExtractRootTrackTransform(FAnimExtractContext(AnimationData.ToPosition), nullptr);
+				}
+			}
+		}
+
+		if (RootBoneTransform.IsSet())
+		{
+			FTransformAnimationAttribute* RootMotionAttribute = Keyframe.Attributes.FindOrAdd<FTransformAnimationAttribute>(AnimMixerComponents->RootTransformAttributeId);
+			RootMotionAttribute->Value = RootBoneTransform.GetValue();
+		}
+
+		VM.PushValue(KEYFRAME_STACK_NAME, MakeUnique<FKeyframeState>(MoveTemp(Keyframe)));
+	}
+	else
+	{
+		constexpr bool bIsAdditive = false;
+		FKeyframeState Keyframe = VM.MakeReferenceKeyframe(bIsAdditive);
+		VM.PushValue(KEYFRAME_STACK_NAME, MakeUnique<FKeyframeState>(MoveTemp(Keyframe)));
+	}
+}
+
+UMovieSceneMixedSkeletalAnimationSystem::UMovieSceneMixedSkeletalAnimationSystem(const FObjectInitializer& ObjInit)
+	: Super(ObjInit)
+{
+	using namespace UE::MovieScene;
+
+	FMovieSceneTracksComponentTypes* TrackComponents = FMovieSceneTracksComponentTypes::Get();
+	FAnimMixerComponentTypes* AnimMixerComponents = FAnimMixerComponentTypes::Get();
+	RelevantComponent = TrackComponents->SkeletalAnimation;
+	Phase = ESystemPhase::Instantiation | ESystemPhase::Scheduling;
+	SystemCategories = EEntitySystemCategory::BlenderSystems;
+
+	if (HasAnyFlags(RF_ClassDefaultObject))
+	{
+		FBuiltInComponentTypes* BuiltInComponents = FBuiltInComponentTypes::Get(); 
+
+		DefineImplicitPrerequisite(UMovieSceneEvalTimeSystem::StaticClass(), GetClass());
+		DefineImplicitPrerequisite(UMovieSceneGenericBoundObjectInstantiator::StaticClass(), GetClass());
+		DefineImplicitPrerequisite(UMovieSceneBoundSceneComponentInstantiator::StaticClass(), GetClass());
+		DefineComponentConsumer(GetClass(), BuiltInComponents->BoundObject);
+		DefineComponentConsumer(GetClass(), TrackComponents->SkeletalAnimation);
+		DefineComponentProducer(GetClass(), TrackComponents->Tags.AnimMixerPoseProducer);
+		DefineImplicitPrerequisite(GetClass(), UMovieSceneRestorePreAnimatedStateSystem::StaticClass());
+		DefineComponentProducer(GetClass(), AnimMixerComponents->Task);
+		DefineComponentProducer(GetClass(), AnimMixerComponents->Target);
+		DefineComponentProducer(GetClass(), AnimMixerComponents->Priority);
+
+	}
+	SystemCategories |= FSystemInterrogator::GetExcludedFromInterrogationCategory();
+}
+
+void UMovieSceneMixedSkeletalAnimationSystem::OnRun(FSystemTaskPrerequisites& InPrerequisites, FSystemSubsequentTasks& Subsequents)
+{
+	using namespace UE::MovieScene;
+
+	FBuiltInComponentTypes* BuiltInComponents = FBuiltInComponentTypes::Get();
+	FMovieSceneTracksComponentTypes* TrackComponents = FMovieSceneTracksComponentTypes::Get();
+	FAnimMixerComponentTypes* AnimMixerComponents = FAnimMixerComponentTypes::Get();
+	TSharedRef<FMovieSceneEntitySystemRunner> Runner = Linker->GetRunner();
+
+	if (Runner->GetCurrentPhase() == ESystemPhase::Instantiation)
+	{
+		FEntityComponentFilter Filter;
+		Filter.All({ BuiltInComponents->BoundObject, TrackComponents->SkeletalAnimation, BuiltInComponents->Tags.NeedsLink });
+		Filter.None({ BuiltInComponents->Tags.ImportedEntity });
+
+		// Initialize components for skeletal animation mixing
+		FSkeletalAnimMixerMutation SkeletalAnimMixerMutation;
+		Linker->EntityManager.MutateAll(Filter, SkeletalAnimMixerMutation);
+	}
+}
+
+void UMovieSceneMixedSkeletalAnimationSystem::OnSchedulePersistentTasks(UE::MovieScene::IEntitySystemScheduler* TaskScheduler)
+{
+	using namespace UE::MovieScene;
+
+	FBuiltInComponentTypes* BuiltInComponents = FBuiltInComponentTypes::Get();
+	FMovieSceneTracksComponentTypes* TrackComponents = FMovieSceneTracksComponentTypes::Get();
+	FAnimMixerComponentTypes* AnimMixerComponents = FAnimMixerComponentTypes::Get();
+
+	// Gather all skel anim section entities that have been marked with the AnimMixerPoseProducer tag
+	FTaskID GatherTask = FEntityTaskBuilder()
+		.ReadEntityIDs()
+		.Read(BuiltInComponents->InstanceHandle)
+		.Read(BuiltInComponents->BoundObject)
+		.Read(TrackComponents->SkeletalAnimation)
+		.ReadOptional(BuiltInComponents->EvalTime)
+		.Write(AnimMixerComponents->Task)
+		.FilterAll( { TrackComponents->Tags.AnimMixerPoseProducer })
+		.FilterNone({ BuiltInComponents->Tags.Ignored})
+		.Schedule_PerAllocation<FGatherMixableSkeletalAnimations>(&Linker->EntityManager, TaskScheduler,
+			Linker->GetInstanceRegistry());
+}
diff --git a/Engine/Plugins/Experimental/MovieSceneAnimMixer/Source/MovieSceneAnimMixer/Private/Systems/MovieSceneRootMotionSystem.cpp b/Engine/Plugins/Experimental/MovieSceneAnimMixer/Source/MovieSceneAnimMixer/Private/Systems/MovieSceneRootMotionSystem.cpp
new file mode 100644
--- /dev/null
+++ b/Engine/Plugins/Experimental/MovieSceneAnimMixer/Source/MovieSceneAnimMixer/Private/Systems/MovieSceneRootMotionSystem.cpp
@@ -0,0 +1,179 @@
+// Copyright Epic Games, Inc. All Rights Reserved.
+
+#include "Systems/MovieSceneRootMotionSystem.h"
+
+#include "EntitySystem/BuiltInComponentTypes.h"
+#include "EntitySystem/IMovieSceneTaskScheduler.h"
+#include "MovieSceneTracksComponentTypes.h"
+
+#include "EvaluationVM/EvaluationTask.h"
+#include "EvaluationVM/EvaluationVM.h"
+
+#include "Systems/MovieSceneAnimMixerSystem.h"
+#include "AnimMixerComponentTypes.h"
+
+#include "Components/SkeletalMeshComponent.h"
+#include "Graph/AnimNext_LODPose.h"
+#include "Animation/BuiltInAttributeTypes.h"
+#include "EntitySystem/Interrogation/MovieSceneInterrogationLinker.h"
+#include "Systems/MovieScenePiecewiseDoubleBlenderSystem.h"
+#include "Misc/MemStack.h"
+#include "EntitySystem/MovieSceneEntitySystemTask.h"
+#include "EntitySystem/MovieSceneEntitySystemLinker.h"
+#include "Systems/MovieScenePropertyInstantiator.h"
+#include "EntitySystem/Interrogation/MovieSceneInterrogatedPropertyInstantiator.h"
+
+#include UE_INLINE_GENERATED_CPP_BY_NAME(MovieSceneRootMotionSystem)
+
+namespace UE::MovieScene
+{
+	struct FEvaluateRootMotion
+	{
+		void ForEachEntity
+		(
+			FObjectComponent MeshComponent,
+			TSharedPtr<FAnimNextEvaluationTask> MixerTask,
+			ESwapRootBone RootMotionType,
+			double& DoubleResult0,
+			double& DoubleResult1,
+			double& DoubleResult2,
+			double& DoubleResult3,
+			double& DoubleResult4,
+			double& DoubleResult5,
+			double& DoubleResult6,
+			double& DoubleResult7,
+			double& DoubleResult8
+		) const
+		{
+			using namespace UE::AnimNext;
+			using namespace UE::MovieScene;
+
+			// Zero out the root motion to start
+			DoubleResult0 = 0;
+			DoubleResult1 = 0;
+			DoubleResult2 = 0;
+			DoubleResult3 = 0;
+			DoubleResult4 = 0;
+			DoubleResult5 = 0;
+			DoubleResult6 = 0;
+			DoubleResult7 = 0;
+			DoubleResult8 = 0;
+
+			if (RootMotionType == ESwapRootBone::SwapRootBone_None)
+			{
+				
+				return;
+			}
+
+			FAnimMixerComponentTypes* AnimMixerTypes = FAnimMixerComponentTypes::Get();
+
+			USkeletalMeshComponent* SkeletalMeshComponent = CastChecked<USkeletalMeshComponent>(MeshComponent.GetObject());
+			ensureMsgf(SkeletalMeshComponent, TEXT("Attempting to apply animation to an anim instance without a valid skeletal mesh component."));
+
+			if (!SkeletalMeshComponent || !SkeletalMeshComponent->GetSkeletalMeshAsset())
+			{
+				return;
+			}
+
+			if (MixerTask.IsValid())
+			{
+				FMemMark Mark(FMemStack::Get());
+				FDataHandle RefPoseHandle = FDataRegistry::Get()->GetOrGenerateReferencePose(SkeletalMeshComponent);
+				FAnimNextGraphReferencePose GraphReferencePose(RefPoseHandle);
+
+				const FReferencePose& RefPose = RefPoseHandle.GetRef<FReferencePose>();
+				TOptional<FTransform> ResultTransform;
+				{
+					// Mix just the root motion
+					FEvaluationVM EvaluationVM(EEvaluationFlags::Trajectory, RefPose, 0);
+
+					MixerTask->Execute(EvaluationVM);
+
+					TUniquePtr<FKeyframeState> EvaluatedKeyframe;
+					if (EvaluationVM.PopValue(KEYFRAME_STACK_NAME, EvaluatedKeyframe))
+					{
+						if (EvaluatedKeyframe.IsValid())
+						{
+							if (FTransformAnimationAttribute* RootMotionAttribute = EvaluatedKeyframe->Attributes.Find<FTransformAnimationAttribute>(AnimMixerTypes->RootTransformAttributeId))
+							{
+								ResultTransform = RootMotionAttribute->Value;
+							}
+						}
+					}
+				}
+
+				if (ResultTransform.IsSet())
+				{
+					if (RootMotionType == ESwapRootBone::SwapRootBone_Actor)
+					{
+						AActor* Actor = SkeletalMeshComponent->GetOwner();
+						if (Actor && Actor->GetRootComponent())
+						{
+							// Compensate for any mesh component rotation
+							FQuat InverseMeshToActorRotation = SkeletalMeshComponent->GetOwner()->GetRootComponent()->GetComponentTransform().GetRelativeTransformReverse(SkeletalMeshComponent->GetComponentTransform()).GetRotation();
+							ResultTransform->SetTranslation(InverseMeshToActorRotation * ResultTransform->GetTranslation());
+						}
+					}
+
+					DoubleResult0 = ResultTransform->GetTranslation().X;
+					DoubleResult1 = ResultTransform->GetTranslation().Y;
+					DoubleResult2 = ResultTransform->GetTranslation().Z;
+
+					FVector Euler = ResultTransform->Rotator().Euler();
+					DoubleResult3 = Euler.X;
+					DoubleResult4 = Euler.Y;
+					DoubleResult5 = Euler.Z;
+
+					// For now ignore scale. There's probably a better way to do this.
+					DoubleResult6 = 0; //ResultTransform->GetScale3D().X;
+					DoubleResult7 = 0; //ResultTransform->GetScale3D().Y;
+					DoubleResult8 = 0; //ResultTransform->GetScale3D().Z;*/
+				}
+			}
+		}
+	};
+};
+
+UMovieSceneRootMotionSystem::UMovieSceneRootMotionSystem(const FObjectInitializer& ObjInit)
+	: Super(ObjInit)
+{
+	using namespace UE::MovieScene;
+
+	FAnimMixerComponentTypes* AnimMixerComponents = FAnimMixerComponentTypes::Get();
+
+	RelevantComponent = AnimMixerComponents->MixerTask;
+	Phase = ESystemPhase::Scheduling;
+
+	SystemCategories |= FSystemInterrogator::GetExcludedFromInterrogationCategory();
+
+	if (HasAnyFlags(RF_ClassDefaultObject))
+	{
+		DefineImplicitPrerequisite(GetClass(), UMovieScenePiecewiseDoubleBlenderSystem::StaticClass());
+		DefineImplicitPrerequisite(UMovieSceneAnimMixerSystem::StaticClass(), GetClass()); 
+	}
+}
+
+void UMovieSceneRootMotionSystem::OnSchedulePersistentTasks(UE::MovieScene::IEntitySystemScheduler* TaskScheduler)
+{
+	using namespace UE::MovieScene;
+
+	FBuiltInComponentTypes* BuiltInComponents = FBuiltInComponentTypes::Get();
+	FAnimMixerComponentTypes* AnimMixerComponents = FAnimMixerComponentTypes::Get();
+
+	// Evaluate root motion
+	FTaskID EvaluateTask = FEntityTaskBuilder()
+		.Read(AnimMixerComponents->MeshComponent)
+		.Read(AnimMixerComponents->MixerTask)
+		.Read(AnimMixerComponents->RootMotionType)
+		.Write(BuiltInComponents->DoubleResult[0])
+		.Write(BuiltInComponents->DoubleResult[1])
+		.Write(BuiltInComponents->DoubleResult[2])
+		.Write(BuiltInComponents->DoubleResult[3])
+		.Write(BuiltInComponents->DoubleResult[4])
+		.Write(BuiltInComponents->DoubleResult[5])
+		.Write(BuiltInComponents->DoubleResult[6])
+		.Write(BuiltInComponents->DoubleResult[7])
+		.Write(BuiltInComponents->DoubleResult[8])
+		.SetDesiredThread(Linker->EntityManager.GetGatherThread())
+		.Schedule_PerEntity<FEvaluateRootMotion>(&Linker->EntityManager, TaskScheduler);
+}
\ No newline at end of file
diff --git a/Engine/Plugins/Experimental/MovieSceneAnimMixer/Source/MovieSceneAnimMixer/Public/AnimMixerComponentTypes.h b/Engine/Plugins/Experimental/MovieSceneAnimMixer/Source/MovieSceneAnimMixer/Public/AnimMixerComponentTypes.h
new file mode 100644
--- /dev/null
+++ b/Engine/Plugins/Experimental/MovieSceneAnimMixer/Source/MovieSceneAnimMixer/Public/AnimMixerComponentTypes.h
@@ -0,0 +1,42 @@
+// Copyright Epic Games, Inc. All Rights Reserved.
+
+#pragma once
+
+#include "EntitySystem/BuiltInComponentTypes.h"
+#include "EntitySystem/MovieSceneEntityIDs.h"
+#include "EntitySystem/MovieScenePropertyMetaDataTraits.h"
+#include "EntitySystem/MovieScenePropertySystemTypes.h"
+#include "EntitySystem/MovieScenePropertyTraits.h"
+#include "MovieSceneMixedAnimationTarget.h"
+#include "StructUtils/InstancedStruct.h"
+#include "EvaluationVM/EvaluationTask.h"
+#include "AnimSequencerInstanceProxy.h"
+
+namespace UE::MovieScene
+{
+	struct FAnimMixerComponentTypes
+	{
+		MOVIESCENEANIMMIXER_API static FAnimMixerComponentTypes* Get();
+		MOVIESCENEANIMMIXER_API static void Destroy();
+
+		// A root motion transform in animation space to be applied (after space conversion)
+		static const UE::Anim::FAttributeId RootTransformAttributeId;
+
+		TComponentTypeID<int32> Priority;
+		TComponentTypeID<TInstancedStruct<FMovieSceneMixedAnimationTarget>> Target;
+		TComponentTypeID<TSharedPtr<FAnimNextEvaluationTask>> Task;
+		TComponentTypeID<TSharedPtr<FAnimNextEvaluationTask>> MixerTask;
+		TComponentTypeID<ESwapRootBone> RootMotionType;
+		TComponentTypeID<FObjectComponent> MeshComponent;
+
+		struct
+		{
+			FComponentTypeID RequiresBlending;
+		} Tags;
+
+
+	private:
+		FAnimMixerComponentTypes();
+	};
+
+}
\ No newline at end of file
diff --git a/Engine/Plugins/Experimental/MovieSceneAnimMixer/Source/MovieSceneAnimMixer/Public/Systems/MovieSceneAnimInstanceTargetSystem.h b/Engine/Plugins/Experimental/MovieSceneAnimMixer/Source/MovieSceneAnimMixer/Public/Systems/MovieSceneAnimInstanceTargetSystem.h
new file mode 100644
--- /dev/null
+++ b/Engine/Plugins/Experimental/MovieSceneAnimMixer/Source/MovieSceneAnimMixer/Public/Systems/MovieSceneAnimInstanceTargetSystem.h
@@ -0,0 +1,129 @@
+// Copyright Epic Games, Inc. All Rights Reserved.
+
+#pragma once
+
+#include "CoreMinimal.h"
+#include "EntitySystem/MovieSceneEntitySystem.h"
+#include "MovieSceneMixedAnimationTarget.h"
+#include "Animation/AnimInstance.h"
+#include "Animation/AnimInstanceProxy.h"
+#include "SequencerAnimationSupport.h"
+#include "Animation/AnimNodeBase.h"
+#include "MovieSceneAnimInstanceTargetSystem.generated.h"
+
+
+struct FAnimNextEvaluationTask;
+
+/**
+ * Declaring a unique target for targeting a custom anim instance on the skeletal mesh component. Doesn't need additional metadata.
+ */
+USTRUCT(meta=(DisplayName="Custom Anim Instance"))
+struct MOVIESCENEANIMMIXER_API FMovieSceneAnimInstanceTarget : public FMovieSceneMixedAnimationTarget
+{
+	GENERATED_BODY() 
+
+	FORCEINLINE friend uint32 GetTypeHash(const FMovieSceneAnimInstanceTarget& Target)
+	{
+		return GetTypeHash(FMovieSceneAnimInstanceTarget::StaticStruct());
+	}
+};
+
+
+// Custom anim instance and proxy that take an anim evaluation task, evaluate it, and push the resulting pose to the skeleton.
+UCLASS(transient, NotBlueprintable, MinimalAPI)
+class USequencerMixedAnimInstance : public UAnimInstance, public ISequencerAnimationSupport
+{
+	GENERATED_BODY()
+
+	USequencerMixedAnimInstance(const FObjectInitializer& ObjectInitializer);
+
+public:
+	MOVIESCENEANIMMIXER_API void SetMixerTask(TSharedPtr<FAnimNextEvaluationTask> InEvalTask);
+
+protected:
+	// UAnimInstance interface
+	MOVIESCENEANIMMIXER_API virtual FAnimInstanceProxy* CreateAnimInstanceProxy() override;
+
+	// ISequencerAnimationSupport
+
+	// Empty ISequencerAnimationSupport things we don't use- would be better to have 2 interfaces, but deprecation would be a pain.
+	MOVIESCENEANIMMIXER_API virtual void UpdateAnimTrack(UAnimSequenceBase* InAnimSequence, int32 SequenceId, float InPosition, float Weight, bool bFireNotifies) override {}
+	MOVIESCENEANIMMIXER_API virtual void UpdateAnimTrack(UAnimSequenceBase* InAnimSequence, int32 SequenceId, float InFromPosition, float InToPosition, float Weight, bool bFireNotifies) override {}
+	MOVIESCENEANIMMIXER_API virtual void ConstructNodes() override {}
+	MOVIESCENEANIMMIXER_API virtual void ResetNodes() override {}
+	MOVIESCENEANIMMIXER_API virtual void ResetPose() override {}
+	MOVIESCENEANIMMIXER_API virtual void SavePose() override {}
+
+	virtual UAnimInstance* GetSourceAnimInstance() override;
+	virtual void SetSourceAnimInstance(UAnimInstance* SourceAnimInstance) override;
+	virtual bool DoesSupportDifferentSourceAnimInstance() const override { return true; }
+};
+
+
+/** Proxy override for this UAnimInstance-derived class */
+USTRUCT()
+struct FSequencerMixedAnimInstanceProxy : public FAnimInstanceProxy
+{
+	GENERATED_BODY()
+
+public:
+	FSequencerMixedAnimInstanceProxy()
+	{
+	}
+
+	FSequencerMixedAnimInstanceProxy(UAnimInstance* InAnimInstance)
+		: FAnimInstanceProxy(InAnimInstance)
+	{
+	}
+
+	MOVIESCENEANIMMIXER_API virtual ~FSequencerMixedAnimInstanceProxy() {}
+
+	// FAnimInstanceProxy interface
+	MOVIESCENEANIMMIXER_API virtual void Initialize(UAnimInstance* InAnimInstance) override;
+	MOVIESCENEANIMMIXER_API virtual bool Evaluate(FPoseContext& Output) override;
+	MOVIESCENEANIMMIXER_API virtual void CacheBones() override;
+	MOVIESCENEANIMMIXER_API virtual void UpdateAnimationNode(const FAnimationUpdateContext& InContext) override;
+	MOVIESCENEANIMMIXER_API virtual void PreEvaluateAnimation(UAnimInstance* InAnimInstance) override;
+
+	/** Anim Instance Source info - created externally and used here */
+	void SetSourceAnimInstance(UAnimInstance* SourceAnimInstance, FAnimInstanceProxy* SourceAnimInputProxy);
+	UAnimInstance* GetSourceAnimInstance() const { return CurrentSourceAnimInstance; }
+
+	void SetMixerTask(TSharedPtr<FAnimNextEvaluationTask> InEvalTask);
+
+
+	void LinkSourcePose(UAnimInstance* InInputInstance, FAnimInstanceProxy* InInputProxy);
+	void UnlinkSourcePose();
+
+protected:
+
+	/** Source Anim Instance */
+	TObjectPtr<UAnimInstance> CurrentSourceAnimInstance;
+	FAnimInstanceProxy* CurrentSourceProxy = nullptr;
+
+	/* Optional link to root node if the source anim instance exists and has one */
+	FPoseLink SourcePose;
+
+	// Pointer to the task passed from the mixer.
+	TSharedPtr<FAnimNextEvaluationTask> MixerTask;
+};
+
+
+
+
+// System that handles applying animation mixer evaluation tasks to a custom anim instance on a skeletal mesh component.
+UCLASS(MinimalAPI)
+class UMovieSceneAnimInstanceTargetSystem
+	: public UMovieSceneEntitySystem
+{
+public:
+
+	GENERATED_BODY()
+
+	UMovieSceneAnimInstanceTargetSystem(const FObjectInitializer& ObjInit);
+
+private:
+
+	virtual void OnSchedulePersistentTasks(UE::MovieScene::IEntitySystemScheduler* TaskScheduler) override;
+
+};
\ No newline at end of file
diff --git a/Engine/Plugins/Experimental/MovieSceneAnimMixer/Source/MovieSceneAnimMixer/Public/Systems/MovieSceneAnimMixerSystem.h b/Engine/Plugins/Experimental/MovieSceneAnimMixer/Source/MovieSceneAnimMixer/Public/Systems/MovieSceneAnimMixerSystem.h
new file mode 100644
--- /dev/null
+++ b/Engine/Plugins/Experimental/MovieSceneAnimMixer/Source/MovieSceneAnimMixer/Public/Systems/MovieSceneAnimMixerSystem.h
@@ -0,0 +1,98 @@
+// Copyright Epic Games, Inc. All Rights Reserved.
+
+#pragma once
+
+#include "EntitySystem/MovieSceneEntitySystem.h"
+#include "MovieSceneMixedAnimationTarget.h"
+#include "StructUtils/InstancedStruct.h"
+#include "Containers/ContainersFwd.h"
+#include "Containers/Map.h"
+#include "Templates/SharedPointer.h"
+#include "AnimSequencerInstanceProxy.h"
+#include "MovieSceneAnimMixerSystem.generated.h"
+
+struct FAnimNextEvaluationTask;
+
+namespace UE
+{
+	namespace MovieScene
+	{
+		struct EntityID;
+	}
+	namespace AnimNext
+	{
+		struct FEvaluationProgram;
+	}
+}
+
+
+struct FMovieSceneAnimMixerEntry
+{
+	UE::MovieScene::FMovieSceneEntityID EntityID;
+	TSharedPtr<FAnimNextEvaluationTask> EvalTask;
+	double PoseWeight;
+	int32 Priority;
+	bool bAdditive = false;
+	bool bRequiresBlend = true;
+	ESwapRootBone RootMotionType = ESwapRootBone::SwapRootBone_None;
+
+	FORCEINLINE bool operator<(const FMovieSceneAnimMixerEntry& RHS) const
+	{
+		return Priority < RHS.Priority;
+	}
+};
+
+// Key into the mixer map- one mixer per bound object per animation target
+
+struct FMovieSceneAnimMixerKey
+{
+	FObjectKey BoundObjectKey;
+	TInstancedStruct<FMovieSceneMixedAnimationTarget> Target;		
+
+	/** Equality operator */
+	friend bool operator==(const FMovieSceneAnimMixerKey& A, const FMovieSceneAnimMixerKey& B)
+	{
+		return A.BoundObjectKey == B.BoundObjectKey && A.Target == B.Target;
+	}
+
+	/** Generate a type hash from this component */
+	friend uint32 GetTypeHash(const FMovieSceneAnimMixerKey& AnimMixerKey)
+	{
+		return HashCombineFast(GetTypeHash(AnimMixerKey.BoundObjectKey), AnimMixerKey.Target.IsValid() ? AnimMixerKey.Target.GetScriptStruct()->GetStructTypeHash(&AnimMixerKey.Target.Get()) : INDEX_NONE);
+	}
+};
+
+// A single anim mixer, containing an array of mixer entries sorted by priority.
+
+struct FMovieSceneAnimMixer
+{
+	UE::MovieScene::FMovieSceneEntityID MixerEntityID;
+	TArray<FMovieSceneAnimMixerEntry> Entries;
+	TSharedPtr<UE::AnimNext::FEvaluationProgram> EvaluationProgram;
+};
+
+// Takes in evaluation tasks with optional pose weight, masks, priority and a given animation target.
+// Constructs a hierarchical 'mixer' per bound object per target. 
+// Similar to blender systems, in a 'many to one' operation, each mixer will create an entity with a single evaluation task
+// wrapping the full blend operation, with the target component. 
+// This entity is then consumed by the appropriate target animation system in order to produce the result on the mesh.
+UCLASS(MinimalAPI)
+class UMovieSceneAnimMixerSystem
+	: public UMovieSceneEntitySystem
+{
+public:
+
+	GENERATED_BODY()
+
+	UMovieSceneAnimMixerSystem(const FObjectInitializer& ObjInit);
+
+	// Map of animation mixers
+	TMap<FMovieSceneAnimMixerKey, FMovieSceneAnimMixer> Mixers;
+
+private:
+
+	virtual void OnRun(FSystemTaskPrerequisites& InPrerequisites, FSystemSubsequentTasks& Subsequents) override;
+	virtual void OnSchedulePersistentTasks(UE::MovieScene::IEntitySystemScheduler* TaskScheduler) override;
+	virtual void OnCleanTaggedGarbage() override;
+
+};
\ No newline at end of file
diff --git a/Engine/Plugins/Experimental/MovieSceneAnimMixer/Source/MovieSceneAnimMixer/Public/Systems/MovieSceneMixedSkeletalAnimationSystem.h b/Engine/Plugins/Experimental/MovieSceneAnimMixer/Source/MovieSceneAnimMixer/Public/Systems/MovieSceneMixedSkeletalAnimationSystem.h
new file mode 100644
--- /dev/null
+++ b/Engine/Plugins/Experimental/MovieSceneAnimMixer/Source/MovieSceneAnimMixer/Public/Systems/MovieSceneMixedSkeletalAnimationSystem.h
@@ -0,0 +1,94 @@
+// Copyright Epic Games, Inc. All Rights Reserved.
+
+#pragma once
+
+#include "EntitySystem/MovieSceneEntitySystem.h"
+#include "EntitySystem/MovieSceneSequenceInstanceHandle.h"
+#include "Evaluation/MovieScenePlayback.h"
+#include "UObject/ObjectKey.h"
+#include "EntitySystem/MovieSceneEntityInstantiatorSystem.h"
+#include "EvaluationVM/EvaluationTask.h"
+#include "AnimSequencerInstanceProxy.h"
+#include "MovieSceneMixedSkeletalAnimationSystem.generated.h"
+
+struct FAnimSequencerData;
+class UAnimSequence;
+
+USTRUCT()
+struct FSkeletalAnimationRootMotionOverride
+{
+	GENERATED_BODY()
+
+	/** If true we use the ChildBoneIndex otherwise we use the root*/
+	UPROPERTY()
+	bool bBlendFirstChildOfRoot = false;
+	
+	UPROPERTY()
+	int32 ChildBoneIndex = INDEX_NONE;
+	
+	UPROPERTY()
+	FTransform RootMotion = FTransform::Identity;
+	
+	UPROPERTY()
+	FTransform PreviousTransform = FTransform::Identity;
+};
+
+USTRUCT()
+struct FMixedAnimSkeletalAnimationData
+{
+	GENERATED_BODY()
+
+	UPROPERTY()
+	TWeakObjectPtr<UAnimSequence> AnimSequence;
+
+	UPROPERTY()
+	TOptional<FSkeletalAnimationRootMotionOverride> RootMotionOverride;
+
+	UPROPERTY()
+	double FromPosition;
+
+	UPROPERTY()
+	double ToPosition;
+
+	UPROPERTY()
+	bool bFireNotifies;
+
+	UPROPERTY()
+	ESwapRootBone SwapRootBone;
+
+	UPROPERTY()
+	bool bAdditive;
+};
+
+USTRUCT()
+struct FMovieSceneSkeletalAnimationEvaluationTask : public FAnimNextEvaluationTask
+{
+	GENERATED_BODY()
+
+	DECLARE_ANIM_EVALUATION_TASK(FMovieSceneSkeletalAnimationEvaluationTask)
+
+	// Task entry point
+	virtual void Execute(UE::AnimNext::FEvaluationVM& VM) const override;
+
+	UPROPERTY()
+	FMixedAnimSkeletalAnimationData AnimationData;
+};
+
+// System to handle creating evaluation tasks from skeletal animation track sections for the anim mixer.
+UCLASS(MinimalAPI)
+class UMovieSceneMixedSkeletalAnimationSystem
+	: public UMovieSceneEntityInstantiatorSystem
+{
+public:
+
+	GENERATED_BODY()
+
+	UMovieSceneMixedSkeletalAnimationSystem(const FObjectInitializer& ObjInit);
+
+private:
+
+	virtual void OnRun(FSystemTaskPrerequisites& InPrerequisites, FSystemSubsequentTasks& Subsequents) override;
+	virtual void OnSchedulePersistentTasks(UE::MovieScene::IEntitySystemScheduler* TaskScheduler) override;
+	
+
+};
\ No newline at end of file
diff --git a/Engine/Plugins/Experimental/MovieSceneAnimMixer/Source/MovieSceneAnimMixer/Public/Systems/MovieSceneRootMotionSystem.h b/Engine/Plugins/Experimental/MovieSceneAnimMixer/Source/MovieSceneAnimMixer/Public/Systems/MovieSceneRootMotionSystem.h
new file mode 100644
--- /dev/null
+++ b/Engine/Plugins/Experimental/MovieSceneAnimMixer/Source/MovieSceneAnimMixer/Public/Systems/MovieSceneRootMotionSystem.h
@@ -0,0 +1,27 @@
+// Copyright Epic Games, Inc. All Rights Reserved.
+
+#pragma once
+
+#include "EntitySystem/MovieSceneEntitySystem.h"
+#include "MovieSceneRootMotionSystem.generated.h"
+
+
+// Takes in evaluation tasks on mixers.
+// Mixes just the root motion attributes.
+// Converts it from animation space to either additive actor or component space (based on which attribute used).
+// Writes it out as an additive transform to be mixed alongside other transform track values.
+UCLASS(MinimalAPI)
+class UMovieSceneRootMotionSystem
+	: public UMovieSceneEntitySystem
+{
+public:
+
+	GENERATED_BODY()
+
+	UMovieSceneRootMotionSystem(const FObjectInitializer& ObjInit);
+
+private:
+
+	virtual void OnSchedulePersistentTasks(UE::MovieScene::IEntitySystemScheduler* TaskScheduler) override;
+
+};
\ No newline at end of file
diff --git a/Engine/Source/Runtime/Engine/Public/Animation/AnimInstanceProxy.h b/Engine/Source/Runtime/Engine/Public/Animation/AnimInstanceProxy.h
--- a/Engine/Source/Runtime/Engine/Public/Animation/AnimInstanceProxy.h
+++ b/Engine/Source/Runtime/Engine/Public/Animation/AnimInstanceProxy.h
@@ -586,6 +586,7 @@ public:
 	friend class UE::Anim::FAnimSyncGroupScope;
 	friend class UE::Anim::FActiveStateMachineScope;
 	friend struct FAnimNode_ControlRigInputPose;
+	friend struct FSequencerMixedAnimInstanceProxy;
 	
 protected:
 	/** Called when our anim instance is being initialized */
diff --git a/Engine/Source/Runtime/MovieSceneTracks/Private/MovieSceneTracksComponentTypes.cpp b/Engine/Source/Runtime/MovieSceneTracks/Private/MovieSceneTracksComponentTypes.cpp
--- a/Engine/Source/Runtime/MovieSceneTracks/Private/MovieSceneTracksComponentTypes.cpp
+++ b/Engine/Source/Runtime/MovieSceneTracks/Private/MovieSceneTracksComponentTypes.cpp
@@ -642,6 +642,8 @@ FMovieSceneTracksComponentTypes::FMovieSceneTracksComponentTypes()
 	// should be interpreted as an index for custom primitive data.
 	Tags.CustomPrimitiveData = ComponentRegistry->NewTag(TEXT("Custom Primitive Data"));
 
+	Tags.AnimMixerPoseProducer = ComponentRegistry->NewTag(TEXT("Anim Mixer Pose Producer"));
+
 	FBuiltInComponentTypes* BuiltInComponents = FBuiltInComponentTypes::Get();
 
 	// --------------------------------------------------------------------------------------------
@@ -898,6 +900,7 @@ FMovieSceneTracksComponentTypes::FMovieSceneTracksComponentTypes()
 	// --------------------------------------------------------------------------------------------
 	// Set up SkeletalAnimation components
 	ComponentRegistry->Factories.DuplicateChildComponent(SkeletalAnimation);
+	ComponentRegistry->Factories.DefineChildComponent(Tags.AnimMixerPoseProducer, Tags.AnimMixerPoseProducer);
 
 	// --------------------------------------------------------------------------------------------
 	// Set up custom primitive data components
diff --git a/Engine/Source/Runtime/MovieSceneTracks/Private/Systems/MovieSceneSkeletalAnimationSystem.cpp b/Engine/Source/Runtime/MovieSceneTracks/Private/Systems/MovieSceneSkeletalAnimationSystem.cpp
--- a/Engine/Source/Runtime/MovieSceneTracks/Private/Systems/MovieSceneSkeletalAnimationSystem.cpp
+++ b/Engine/Source/Runtime/MovieSceneTracks/Private/Systems/MovieSceneSkeletalAnimationSystem.cpp
@@ -1204,7 +1204,7 @@ void UMovieSceneSkeletalAnimationSystem::OnSchedulePersistentTasks(UE::MovieScen
 	.FilterAll({ TrackComponents->ComponentTransform.PropertyTag })
 	.Schedule_PerAllocation<FTransformDependencyTask>(&Linker->EntityManager, TaskScheduler);
 
-
+	// Skip gathering any anims that are tagged as anim mixer pose producer- these will be handled by the version in the anim mixer plugin
 	FTaskID GatherTask = FEntityTaskBuilder()
 	.ReadEntityIDs()
 	.Read(BuiltInComponents->RootInstanceHandle)
@@ -1213,7 +1213,7 @@ void UMovieSceneSkeletalAnimationSystem::OnSchedulePersistentTasks(UE::MovieScen
 	.Read(TrackComponents->SkeletalAnimation)
 	.ReadOptional(BuiltInComponents->EvalTime)
 	.ReadOptional(BuiltInComponents->WeightAndEasingResult)
-	.FilterNone({ BuiltInComponents->Tags.Ignored })
+	.FilterNone({ BuiltInComponents->Tags.Ignored, TrackComponents->Tags.AnimMixerPoseProducer })
 	.SetStat(GET_STATID(MovieSceneEval_GatherSkeletalAnimations))
 	.Schedule_PerAllocation<FGatherSkeletalAnimations>(&Linker->EntityManager, TaskScheduler, 
 			Linker->GetInstanceRegistry(), &SystemData);
@@ -1258,6 +1258,7 @@ void UMovieSceneSkeletalAnimationSystem::OnRun(FSystemTaskPrerequisites& InPrere
 		return;
 	}
 
+	// Skip gathering any anims that are tagged as anim mixer pose producer- these will be handled by the version in the anim mixer plugin
 	FGraphEventRef GatherTask = FEntityTaskBuilder()
 	.ReadEntityIDs()
 	.Read(BuiltInComponents->RootInstanceHandle)
@@ -1266,7 +1267,7 @@ void UMovieSceneSkeletalAnimationSystem::OnRun(FSystemTaskPrerequisites& InPrere
 	.Read(TrackComponents->SkeletalAnimation)
 	.ReadOptional(BuiltInComponents->EvalTime)
 	.ReadOptional(BuiltInComponents->WeightAndEasingResult)
-	.FilterNone({ BuiltInComponents->Tags.Ignored })
+	.FilterNone({ BuiltInComponents->Tags.Ignored, TrackComponents->Tags.AnimMixerPoseProducer })
 	.SetStat(GatherStatId)
 	.Dispatch_PerAllocation<FGatherSkeletalAnimations>(&Linker->EntityManager, InPrerequisites, nullptr, 
 			Linker->GetInstanceRegistry(), &SystemData);
@@ -1285,6 +1286,18 @@ void UMovieSceneSkeletalAnimationSystem::OnRun(FSystemTaskPrerequisites& InPrere
 	.Dispatch<FEvaluateSkeletalAnimations>(&Linker->EntityManager, EvalPrereqs, &Subsequents, Linker, &SystemData);
 }
 
+bool UMovieSceneSkeletalAnimationSystem::IsRelevantImpl(UMovieSceneEntitySystemLinker* InLinker) const
+{
+	using namespace UE::MovieScene;
+	FMovieSceneTracksComponentTypes* TrackComponents = FMovieSceneTracksComponentTypes::Get();
+
+	FEntityComponentFilter Filter;
+	Filter.All({TrackComponents->SkeletalAnimation});
+	Filter.None({TrackComponents->Tags.AnimMixerPoseProducer});
+	return InLinker->EntityManager.Contains(Filter);
+}
+
+
 void UMovieSceneSkeletalAnimationSystem::CleanSystemData()
 {
 	SystemData.ResetSkeletalAnimations();
diff --git a/Engine/Source/Runtime/MovieSceneTracks/Public/MovieSceneMixedAnimationTarget.h b/Engine/Source/Runtime/MovieSceneTracks/Public/MovieSceneMixedAnimationTarget.h
new file mode 100644
--- /dev/null
+++ b/Engine/Source/Runtime/MovieSceneTracks/Public/MovieSceneMixedAnimationTarget.h
@@ -0,0 +1,21 @@
+// Copyright Epic Games, Inc. All Rights Reserved.
+
+#pragma once
+
+#include "CoreTypes.h"
+#include "MovieSceneMixedAnimationTarget.generated.h"
+
+/**
+ * Defines a Mixed Anim Target that animation pose-producing tracks can send pose-producing tasks to. 
+ * Multiple pose producing tasks sent to the same target will be chained and mixed based on priority, pose weight, masks, and sequence hierarchy.
+ * Inheriting from this struct allows the user to create new animation target types with their own metadata and custom component types.
+ * Coupled with a custom component type, these mixer pose results can be read from a custom Movie Scene ECS system and passed as appropriate to an anim system.
+ *  If no target struct is assigned, a default animation target will be assigned based on the makeup of the object being animated.
+ */
+USTRUCT(meta=(DisplayName="Automatic"))
+struct MOVIESCENETRACKS_API FMovieSceneMixedAnimationTarget
+{
+	GENERATED_BODY() 
+
+	virtual ~FMovieSceneMixedAnimationTarget() {}
+};
\ No newline at end of file
diff --git a/Engine/Source/Runtime/MovieSceneTracks/Public/MovieSceneTracksComponentTypes.h b/Engine/Source/Runtime/MovieSceneTracks/Public/MovieSceneTracksComponentTypes.h
--- a/Engine/Source/Runtime/MovieSceneTracks/Public/MovieSceneTracksComponentTypes.h
+++ b/Engine/Source/Runtime/MovieSceneTracks/Public/MovieSceneTracksComponentTypes.h
@@ -787,6 +787,8 @@ struct FMovieSceneTracksComponentTypes
 		FComponentTypeID CustomPrimitiveData;
 		FComponentTypeID Slomo;
 		FComponentTypeID Visibility;
+		// Tag specifying that this entity produces skeleton poses that are evaluated through the anim mixer. 
+		FComponentTypeID AnimMixerPoseProducer;
 	} Tags;
 
 	TComponentTypeID<FLevelVisibilityComponentData> LevelVisibility;
diff --git a/Engine/Source/Runtime/MovieSceneTracks/Public/Sections/MovieSceneSkeletalAnimationSection.h b/Engine/Source/Runtime/MovieSceneTracks/Public/Sections/MovieSceneSkeletalAnimationSection.h
--- a/Engine/Source/Runtime/MovieSceneTracks/Public/Sections/MovieSceneSkeletalAnimationSection.h
+++ b/Engine/Source/Runtime/MovieSceneTracks/Public/Sections/MovieSceneSkeletalAnimationSection.h
@@ -8,6 +8,8 @@
 #include "Animation/AnimSequenceBase.h"
 #include "Channels/MovieSceneFloatChannel.h"
 #include "EntitySystem/IMovieSceneEntityProvider.h"
+#include "StructUtils/InstancedStruct.h"
+#include "MovieSceneMixedAnimationTarget.h"
 #include "MovieSceneSkeletalAnimationSection.generated.h"
 
 struct FMovieSceneSkeletalAnimRootMotionTrackParams;
@@ -33,12 +35,12 @@ struct FMovieSceneSkeletalAnimationParams
 	 * Convert a sequence frame to a time in seconds inside the animation clip, taking into account start/end offsets,
 	 * looping, etc.
 	 */
-	double MapTimeToAnimation(const UMovieSceneSection* InSection, FFrameTime InPosition, FFrameRate InFrameRate) const;
+	double MOVIESCENETRACKS_API MapTimeToAnimation(const UMovieSceneSection* InSection, FFrameTime InPosition, FFrameRate InFrameRate) const;
 
 	/**
 	 * As above, but with already computed section bounds.
 	 */
-	double MapTimeToAnimation(FFrameNumber InSectionStartTime, FFrameNumber InSectionEndTime, FFrameTime InPosition, FFrameRate InFrameRate) const;
+	double MOVIESCENETRACKS_API MapTimeToAnimation(FFrameNumber InSectionStartTime, FFrameNumber InSectionEndTime, FFrameTime InPosition, FFrameRate InFrameRate) const;
 
 	/**
 	 * Make a transform structure from these animation parameters
@@ -239,6 +241,13 @@ public:
 	UPROPERTY(BlueprintReadWrite, EditAnywhere, Category = "Root Motions")
 	uint8 bShowSkeleton : 1;
 #endif
+
+	UPROPERTY(EditAnywhere, Category = "Mixed Animation")
+	TInstancedStruct<FMovieSceneMixedAnimationTarget> MixedAnimationTarget;
+
+	UPROPERTY(EditAnywhere, Category = "Mixed Animation")
+	int32 MixedAnimationPriority = 0;
+
 	//Previous transform used to specify the global OffsetTransform while calculting the root motions.
 	FTransform PreviousTransform;
 
diff --git a/Engine/Source/Runtime/MovieSceneTracks/Public/Systems/MovieSceneSkeletalAnimationSystem.h b/Engine/Source/Runtime/MovieSceneTracks/Public/Systems/MovieSceneSkeletalAnimationSystem.h
--- a/Engine/Source/Runtime/MovieSceneTracks/Public/Systems/MovieSceneSkeletalAnimationSystem.h
+++ b/Engine/Source/Runtime/MovieSceneTracks/Public/Systems/MovieSceneSkeletalAnimationSystem.h
@@ -115,6 +115,8 @@ private:
 	virtual void OnSchedulePersistentTasks(UE::MovieScene::IEntitySystemScheduler* TaskScheduler) override;
 	virtual void OnRun(FSystemTaskPrerequisites& InPrerequisites, FSystemSubsequentTasks& Subsequents) override final;
 
+	virtual bool IsRelevantImpl(UMovieSceneEntitySystemLinker* InLinker) const override;
+
 	void CleanSystemData();
 
 private:
