Sequencer: Add an Anim Next animation target that uses anim next module injection to process mixed animation. For now, the UX for this target shows the injection site as a 'String' and if left at 'None' it will select the default injection site for whatever module is set up. Long term we likely want to make this easier by allowing the user to search for valid injection sites from modules in the project (or the one on the preview actor).

diff --git a/Engine/Plugins/Experimental/AnimNext/Source/AnimNext/Internal/Component/AnimNextWorldSubsystem.h b/Engine/Plugins/Experimental/AnimNext/Source/AnimNext/Internal/Component/AnimNextWorldSubsystem.h
--- a/Engine/Plugins/Experimental/AnimNext/Source/AnimNext/Internal/Component/AnimNextWorldSubsystem.h
+++ b/Engine/Plugins/Experimental/AnimNext/Source/AnimNext/Internal/Component/AnimNextWorldSubsystem.h
@@ -100,6 +100,10 @@ protected:
 	// Unregister a handle from the subsystem
 	void UnregisterHandle(UE::AnimNext::FModuleHandle& InOutHandle);
 
+	// Returns whether the module represented by the supplied handle is enabled.
+	// If there are pending actions they will take precedent over the actual current state.
+	bool IsHandleEnabled(UE::AnimNext::FModuleHandle InHandle) const;
+
 	// Enables or disables the module represented by the supplied handle
 	// This operation is deferred until the next time the schedule ticks
 	void EnableHandle(UE::AnimNext::FModuleHandle InHandle, bool bInEnabled);
diff --git a/Engine/Plugins/Experimental/AnimNext/Source/AnimNext/Internal/Module/AnimNextModuleInstance.h b/Engine/Plugins/Experimental/AnimNext/Source/AnimNext/Internal/Module/AnimNextModuleInstance.h
--- a/Engine/Plugins/Experimental/AnimNext/Source/AnimNext/Internal/Module/AnimNextModuleInstance.h
+++ b/Engine/Plugins/Experimental/AnimNext/Source/AnimNext/Internal/Module/AnimNextModuleInstance.h
@@ -85,6 +85,9 @@ struct FAnimNextModuleInstance : public FAnimNextDataInterfaceInstance
 		EAnimNextModuleInitMethod InInitMethod);
 	ANIMNEXT_API ~FAnimNextModuleInstance();
 
+	// Checks to see if this entry is ticking
+	bool IsEnabled() const;
+
 	// Enables/disables the ticking of this entry
 	void Enable(bool bInEnabled);
 
diff --git a/Engine/Plugins/Experimental/AnimNext/Source/AnimNext/Private/Component/AnimNextComponent.cpp b/Engine/Plugins/Experimental/AnimNext/Source/AnimNext/Private/Component/AnimNextComponent.cpp
--- a/Engine/Plugins/Experimental/AnimNext/Source/AnimNext/Private/Component/AnimNextComponent.cpp
+++ b/Engine/Plugins/Experimental/AnimNext/Source/AnimNext/Private/Component/AnimNextComponent.cpp
@@ -857,6 +857,15 @@ bool UAnimNextComponent::SetVariableArrayRef(const FName Name, TFunctionRef<void
 	return false;
 }
 
+bool UAnimNextComponent::IsEnabled() const
+{
+	if (Subsystem)
+	{
+		return Subsystem->IsEnabled(this);
+	}
+	return false;
+}
+
 void UAnimNextComponent::SetEnabled(bool bEnabled)
 {
 	if(Subsystem)
diff --git a/Engine/Plugins/Experimental/AnimNext/Source/AnimNext/Private/Component/AnimNextComponentWorldSubsystem.cpp b/Engine/Plugins/Experimental/AnimNext/Source/AnimNext/Private/Component/AnimNextComponentWorldSubsystem.cpp
--- a/Engine/Plugins/Experimental/AnimNext/Source/AnimNext/Private/Component/AnimNextComponentWorldSubsystem.cpp
+++ b/Engine/Plugins/Experimental/AnimNext/Source/AnimNext/Private/Component/AnimNextComponentWorldSubsystem.cpp
@@ -33,6 +33,14 @@ void UAnimNextComponentWorldSubsystem::SetEnabled(UAnimNextComponent* InComponen
 	EnableHandle(InComponent->ModuleHandle, bInEnabled);
 }
 
+bool UAnimNextComponentWorldSubsystem::IsEnabled(const UAnimNextComponent* InComponent) const
+{
+	using namespace UE::AnimNext;
+	check(IsInGameThread());
+	check(InComponent);
+	return IsHandleEnabled(InComponent->ModuleHandle);
+}
+
 #if UE_ENABLE_DEBUG_DRAWING
 void UAnimNextComponentWorldSubsystem::ShowDebugDrawing(UAnimNextComponent* InComponent, bool bInShowDebugDrawing)
 {
diff --git a/Engine/Plugins/Experimental/AnimNext/Source/AnimNext/Private/Component/AnimNextComponentWorldSubsystem.h b/Engine/Plugins/Experimental/AnimNext/Source/AnimNext/Private/Component/AnimNextComponentWorldSubsystem.h
--- a/Engine/Plugins/Experimental/AnimNext/Source/AnimNext/Private/Component/AnimNextComponentWorldSubsystem.h
+++ b/Engine/Plugins/Experimental/AnimNext/Source/AnimNext/Private/Component/AnimNextComponentWorldSubsystem.h
@@ -26,6 +26,9 @@ class UAnimNextComponentWorldSubsystem : public UAnimNextWorldSubsystem
 	// The full release of the module referenced by the component's handle will be deferred after this call is made
 	void Unregister(UAnimNextComponent* InComponent);
 
+	// Returns whether the module is enabled.
+	bool IsEnabled(const UAnimNextComponent* InComponent) const;
+
 	// Enables or disables the module represented by the supplied handle
 	// This operation is deferred until the next time the schedule ticks
 	void SetEnabled(UAnimNextComponent* InComponent, bool bInEnabled);
diff --git a/Engine/Plugins/Experimental/AnimNext/Source/AnimNext/Private/Component/AnimNextWorldSubsystem.cpp b/Engine/Plugins/Experimental/AnimNext/Source/AnimNext/Private/Component/AnimNextWorldSubsystem.cpp
--- a/Engine/Plugins/Experimental/AnimNext/Source/AnimNext/Private/Component/AnimNextWorldSubsystem.cpp
+++ b/Engine/Plugins/Experimental/AnimNext/Source/AnimNext/Private/Component/AnimNextWorldSubsystem.cpp
@@ -170,6 +170,27 @@ void UAnimNextWorldSubsystem::UnregisterHandle(UE::AnimNext::FModuleHandle& InOu
 	}
 }
 
+bool UAnimNextWorldSubsystem::IsHandleEnabled(UE::AnimNext::FModuleHandle InHandle) const
+{
+	using namespace UE::AnimNext;
+	check(IsInGameThread());
+	if (IsValidHandle(InHandle))
+	{
+		// The last pending action takes precedent here if present
+		int32 PendingActionIndex = PendingActions.FindLastByPredicate([&InHandle](const FModulePendingAction& PendingAction) { return PendingAction.Handle == InHandle && PendingAction.Type == FModulePendingAction::EType::EnableHandle;});
+
+		if (PendingActionIndex != INDEX_NONE)
+		{
+			return PendingActions[PendingActionIndex].Payload.Get<bool>();
+		}
+
+		// Otherwise return the current value on the instance
+		const FAnimNextModuleInstance& Instance = Instances.Get(InHandle);
+		return Instance.IsEnabled();
+	}
+	return false;
+}
+
 void UAnimNextWorldSubsystem::EnableHandle(UE::AnimNext::FModuleHandle InHandle, bool bInEnabled)
 {
 	using namespace UE::AnimNext;
diff --git a/Engine/Plugins/Experimental/AnimNext/Source/AnimNext/Private/Module/AnimNextModuleInstance.cpp b/Engine/Plugins/Experimental/AnimNext/Source/AnimNext/Private/Module/AnimNextModuleInstance.cpp
--- a/Engine/Plugins/Experimental/AnimNext/Source/AnimNext/Private/Module/AnimNextModuleInstance.cpp
+++ b/Engine/Plugins/Experimental/AnimNext/Source/AnimNext/Private/Module/AnimNextModuleInstance.cpp
@@ -297,6 +297,15 @@ void FAnimNextModuleInstance::QueueOutputTraitEvent(FAnimNextTraitEventPtr Event
 	OutputEventList.Push(MoveTemp(Event));
 }
 
+bool FAnimNextModuleInstance::IsEnabled() const
+{
+	using namespace UE::AnimNext;
+
+	check(IsInGameThread());
+
+	return RunState == ERunState::Running;
+}
+
 void FAnimNextModuleInstance::Enable(bool bInEnabled)
 {
 	using namespace UE::AnimNext;
diff --git a/Engine/Plugins/Experimental/AnimNext/Source/AnimNext/Public/Component/AnimNextComponent.h b/Engine/Plugins/Experimental/AnimNext/Source/AnimNext/Public/Component/AnimNextComponent.h
--- a/Engine/Plugins/Experimental/AnimNext/Source/AnimNext/Public/Component/AnimNextComponent.h
+++ b/Engine/Plugins/Experimental/AnimNext/Source/AnimNext/Public/Component/AnimNextComponent.h
@@ -181,6 +181,10 @@ public:
 		return SetVariableObject(Name, (UObject*)InValue);
 	}
 
+	// Whether this component is currently updating
+	UFUNCTION(BlueprintCallable, Category = "AnimNext")
+	ANIMNEXT_API bool IsEnabled() const;
+
 	// Enable or disable this component's update
 	UFUNCTION(BlueprintCallable, Category = "AnimNext")
 	ANIMNEXT_API void SetEnabled(bool bEnabled);
diff --git a/Engine/Plugins/Experimental/AnimNextAnimGraph/Source/AnimNextAnimGraph/Public/TraitInterfaces/IEvaluate.h b/Engine/Plugins/Experimental/AnimNextAnimGraph/Source/AnimNextAnimGraph/Public/TraitInterfaces/IEvaluate.h
--- a/Engine/Plugins/Experimental/AnimNextAnimGraph/Source/AnimNextAnimGraph/Public/TraitInterfaces/IEvaluate.h
+++ b/Engine/Plugins/Experimental/AnimNextAnimGraph/Source/AnimNextAnimGraph/Public/TraitInterfaces/IEvaluate.h
@@ -27,6 +27,9 @@ namespace UE::AnimNext
 		template<class TaskType>
 		void AppendTask(TaskType&& Task) { EvaluationProgram.AppendTask(MoveTemp(Task)); }
 
+		template<class TaskType>
+		void AppendTaskPtr(const TSharedPtr<TaskType>& TaskPtr) { EvaluationProgram.AppendTaskPtr(TaskPtr); }
+
 	private:
 		explicit FEvaluateTraversalContext(FEvaluationProgram& InEvaluationProgram);
 
diff --git a/Engine/Plugins/Experimental/MovieSceneAnimMixer/Source/MovieSceneAnimMixer/Private/Systems/MovieSceneAnimNextTargetSystem.cpp b/Engine/Plugins/Experimental/MovieSceneAnimMixer/Source/MovieSceneAnimMixer/Private/Systems/MovieSceneAnimNextTargetSystem.cpp
new file mode 100644
--- /dev/null
+++ b/Engine/Plugins/Experimental/MovieSceneAnimMixer/Source/MovieSceneAnimMixer/Private/Systems/MovieSceneAnimNextTargetSystem.cpp
@@ -0,0 +1,264 @@
+// Copyright Epic Games, Inc. All Rights Reserved.
+
+#include "Systems/MovieSceneAnimNextTargetSystem.h"
+
+#include "EntitySystem/BuiltInComponentTypes.h"
+#include "EntitySystem/Interrogation/MovieSceneInterrogationLinker.h"
+#include "EntitySystem/MovieSceneEntityIDs.h"
+#include "EntitySystem/MovieSceneEntitySystemLinker.h"
+#include "EntitySystem/MovieSceneEntitySystemRunner.h"
+#include "EntitySystem/MovieSceneEntitySystemTask.h"
+#include "MovieSceneTracksComponentTypes.h"
+
+#include "Evaluation/PreAnimatedState/IMovieScenePreAnimatedStorage.h"
+#include "Evaluation/PreAnimatedState/MovieScenePreAnimatedObjectStorage.h"
+#include "Evaluation/PreAnimatedState/MovieScenePreAnimatedStorageID.h"
+#include "Evaluation/PreAnimatedState/MovieScenePreAnimatedStorageID.inl"
+
+#include "EvaluationVM/EvaluationTask.h"
+#include "EvaluationVM/EvaluationProgram.h"
+#include "EvaluationVM/EvaluationVM.h"
+#include "GenerationTools.h"
+
+#include "Systems/MovieSceneComponentTransformSystem.h"
+#include "Systems/MovieSceneQuaternionInterpolationRotationSystem.h"
+#include "Systems/MovieSceneObjectPropertySystem.h"
+#include "EntitySystem/MovieScenePreAnimatedStateSystem.h"
+
+#include "Animation/AnimationPoseData.h"
+#include "BonePose.h"
+#include "Animation/BuiltInAttributeTypes.h"
+#include "Systems/MovieSceneAnimMixerSystem.h"
+#include "AnimMixerComponentTypes.h"
+
+#include "Graph/AnimNext_LODPose.h"
+#include "Injection/InjectionUtils.h"
+#include "TraitInterfaces/IEvaluate.h"
+
+
+#include UE_INLINE_GENERATED_CPP_BY_NAME(MovieSceneAnimNextTargetSystem)
+
+namespace UE::MovieScene
+{
+	/** Pre-animated state for skeletal animations */
+	struct FPreAnimatedAnimNextState
+	{
+		// Handle to the injection request instance (used to uninject when animation done)
+		UE::AnimNext::FInjectionRequestPtr InjectionHandle;
+
+		TObjectPtr<UMovieSceneAnimNextTargetSystem> System;
+
+		// Whether the anim next module started enabled or not.
+		bool bAnimNextModuleEnabled = false;
+	};
+
+	/** Pre-animation traits for skeletal animations */
+	struct FPreAnimatedAnimNextMixedSkeletalAnimationTraits : FBoundObjectPreAnimatedStateTraits
+	{
+		using KeyType = FObjectKey;
+		using StorageType = FPreAnimatedAnimNextState;
+
+		StorageType CachePreAnimatedValue(FObjectKey InKey)
+		{
+			check(false);
+			return FPreAnimatedAnimNextState();
+		}
+
+		void RestorePreAnimatedValue(const KeyType& Object, StorageType& InOutCachedValue, const FRestoreStateParams& Params)
+		{
+			// Uninject
+			if (InOutCachedValue.InjectionHandle.IsValid())
+			{
+				if (InOutCachedValue.System)
+				{
+					InOutCachedValue.System->CurrentTargets.RemoveAll([&InOutCachedValue](const FMovieSceneAnimNextTargetData& TargetData) { return TargetData.InjectionRequestHandle == InOutCachedValue.InjectionHandle;});
+				}
+				UE::AnimNext::FInjectionUtils::Uninject(InOutCachedValue.InjectionHandle);
+			}
+
+			// Reset the anim next module back to not being enabled if applicable
+			if (!InOutCachedValue.bAnimNextModuleEnabled)
+			{
+				if (UAnimNextComponent* AnimNextComponent = Cast<UAnimNextComponent>(Object.ResolveObjectPtr()))
+				{
+					AnimNextComponent->SetEnabled(false);
+				}
+			}
+		}
+	};
+
+	/** Pre-animation storage for anim next target */
+	struct FPreAnimatedAnimNextMixedSkeletalAnimationStorage : TPreAnimatedStateStorage_ObjectTraits<FPreAnimatedAnimNextMixedSkeletalAnimationTraits>
+	{
+		static TAutoRegisterPreAnimatedStorageID<FPreAnimatedAnimNextMixedSkeletalAnimationStorage> StorageID;
+	};
+
+	TAutoRegisterPreAnimatedStorageID<FPreAnimatedAnimNextMixedSkeletalAnimationStorage> FPreAnimatedAnimNextMixedSkeletalAnimationStorage::StorageID;
+
+	/** Task for evaluating and applying animation tasks */
+	struct FEvaluateAnimNextTasks
+	{
+	private:
+
+		UMovieSceneEntitySystemLinker* Linker;
+		UMovieSceneAnimNextTargetSystem* System;
+
+		TSharedPtr<FPreAnimatedAnimNextMixedSkeletalAnimationStorage> PreAnimatedStorage;
+
+	public:
+
+		FEvaluateAnimNextTasks(UMovieSceneEntitySystemLinker* InLinker, UMovieSceneAnimNextTargetSystem* InSystem)
+			: Linker(InLinker)
+			, System(InSystem)
+		{
+			PreAnimatedStorage = InLinker->PreAnimatedState.GetOrCreateStorage<FPreAnimatedAnimNextMixedSkeletalAnimationStorage>();
+		}
+
+		void ForEachEntity
+		(
+			FMovieSceneEntityID EntityID,
+			FObjectComponent ObjectComponent,
+			const TInstancedStruct<FMovieSceneMixedAnimationTarget>& Target,
+			TSharedPtr<FAnimNextEvaluationTask> MixerTask
+		) const
+		{
+			// Invalid or incompatible target, return.
+			const FMovieSceneAnimNextInjectionTarget* InjectionTarget = Target.GetPtr<FMovieSceneAnimNextInjectionTarget>();
+			if (!InjectionTarget)
+			{
+				return;
+			}
+
+			// TODO: It's valid (or will be) to use Anim Next without a Skeletal Mesh Component.
+			// We'll need to refactor this to be able to pass an Anim Next component here and refactor core anim tracks to match.
+
+			UAnimNextComponent* AnimNextComponent = nullptr;
+
+			AnimNextComponent = Cast<UAnimNextComponent>(ObjectComponent.GetObject());
+
+			if (!AnimNextComponent)
+			{
+				AActor* Actor = nullptr;
+				// Could be a skeletal mesh component
+				if (UActorComponent* ActorComponent = Cast<UActorComponent>(ObjectComponent.GetObject()))
+				{
+					Actor = ActorComponent->GetOwner();
+				}
+				else
+				{
+					Actor = Cast<AActor>(ObjectComponent.GetObject());
+				}
+
+				if (Actor)
+				{
+					AnimNextComponent = Actor->FindComponentByClass<UAnimNextComponent>();
+				}
+			}
+
+			if (AnimNextComponent)
+			{
+				// If not currently injected, inject
+				FMovieSceneAnimNextTargetData* TargetData = System->CurrentTargets.FindByPredicate([AnimNextComponent, InjectionTarget](const FMovieSceneAnimNextTargetData& Data) { return Data.AnimNextComponent == AnimNextComponent && Data.InjectionSiteName == InjectionTarget->InjectionSiteName;});
+
+				if (!TargetData)
+				{
+					TargetData = &System->CurrentTargets.AddDefaulted_GetRef();
+					TargetData->AnimNextComponent = AnimNextComponent;
+					TargetData->InjectionSiteName = InjectionTarget->InjectionSiteName;
+					TargetData->Modifier = MakeShared<FMovieSceneAnimMixerEvaluationModifier>(MixerTask);
+
+					// Inject
+					TargetData->InjectionRequestHandle = UE::AnimNext::FInjectionUtils::InjectEvaluationModifier(AnimNextComponent, TargetData->Modifier.ToSharedRef(), TargetData->InjectionSiteName);
+
+					// For now, we always restore state
+					PreAnimatedStorage->BeginTrackingEntity(EntityID, true, FRootInstanceHandle(), AnimNextComponent);
+
+					bool bAnimNextModuleEnabled = AnimNextComponent->IsEnabled();
+
+					FCachePreAnimatedValueParams CacheParams;
+					PreAnimatedStorage->CachePreAnimatedValue(
+						CacheParams,
+						AnimNextComponent,
+						[InjectionHandle = TargetData->InjectionRequestHandle, AnimSystem=System, bAnimNextModuleEnabled](FObjectKey InKey) { return FPreAnimatedAnimNextState{InjectionHandle, AnimSystem, bAnimNextModuleEnabled }; });
+
+					// Force enable the anim next module if not currently enabled
+					if (!bAnimNextModuleEnabled)
+					{
+						AnimNextComponent->SetEnabled(true);
+					}
+				}
+				else if(TargetData->Modifier.IsValid())
+				{
+					// Update Task
+					TargetData->Modifier->TaskToInject = MixerTask;
+				}
+			}
+
+		}
+	};
+
+} // namespace UE::MovieScene
+
+void FMovieSceneAnimMixerEvaluationModifier::PostEvaluate(UE::AnimNext::FEvaluateTraversalContext& Context) const
+{
+	if (TaskToInject.IsValid())
+	{
+		Context.AppendTaskPtr(TaskToInject);
+	}
+}
+
+UMovieSceneAnimNextTargetSystem::UMovieSceneAnimNextTargetSystem(const FObjectInitializer& ObjInit)
+	: Super(ObjInit)
+{
+	using namespace UE::MovieScene;
+
+	FMovieSceneTracksComponentTypes* TrackComponents = FMovieSceneTracksComponentTypes::Get();
+	FAnimMixerComponentTypes* AnimMixerComponents = FAnimMixerComponentTypes::Get();
+
+	RelevantComponent = AnimMixerComponents->MixerTask;
+	Phase = ESystemPhase::Scheduling;
+
+	SystemCategories |= FSystemInterrogator::GetExcludedFromInterrogationCategory();
+
+	if (HasAnyFlags(RF_ClassDefaultObject))
+	{
+		DefineImplicitPrerequisite(UMovieSceneAnimMixerSystem::StaticClass(), GetClass());
+		DefineImplicitPrerequisite(GetClass(), UMovieSceneRestorePreAnimatedStateSystem::StaticClass());
+	}
+}
+
+void UMovieSceneAnimNextTargetSystem::OnUnlink()
+{
+	// Clean up system data
+	for (const FMovieSceneAnimNextTargetData& TargetData : CurrentTargets)
+	{
+		// Remove injection
+		if (TargetData.InjectionRequestHandle.IsValid())
+		{
+			UE::AnimNext::FInjectionUtils::Uninject(TargetData.InjectionRequestHandle);
+		}
+	}
+
+	CurrentTargets.Empty();
+}
+
+void UMovieSceneAnimNextTargetSystem::OnSchedulePersistentTasks(UE::MovieScene::IEntitySystemScheduler* TaskScheduler)
+{
+	using namespace UE::MovieScene;
+
+	FBuiltInComponentTypes* BuiltInComponents = FBuiltInComponentTypes::Get();
+	FMovieSceneTracksComponentTypes* TrackComponents = FMovieSceneTracksComponentTypes::Get();
+	FAnimMixerComponentTypes* AnimMixerComponents = FAnimMixerComponentTypes::Get();
+
+	// Evaluate anim tasks targeting an anim next injection point
+	FTaskParams Params(TEXT("Apply AnimNext Animation Tasks"));
+	Params.ForceGameThread();
+	FTaskID EvaluateTask = FEntityTaskBuilder()
+		.ReadEntityIDs()
+		.Read(AnimMixerComponents->MeshComponent)
+		.Read(AnimMixerComponents->Target)
+		.Read(AnimMixerComponents->MixerTask)
+		.SetParams(Params)
+		.Schedule_PerEntity<FEvaluateAnimNextTasks>(&Linker->EntityManager, TaskScheduler,
+			Linker, this);
+}
\ No newline at end of file
diff --git a/Engine/Plugins/Experimental/MovieSceneAnimMixer/Source/MovieSceneAnimMixer/Public/Systems/MovieSceneAnimNextTargetSystem.h b/Engine/Plugins/Experimental/MovieSceneAnimMixer/Source/MovieSceneAnimMixer/Public/Systems/MovieSceneAnimNextTargetSystem.h
new file mode 100644
--- /dev/null
+++ b/Engine/Plugins/Experimental/MovieSceneAnimMixer/Source/MovieSceneAnimMixer/Public/Systems/MovieSceneAnimNextTargetSystem.h
@@ -0,0 +1,77 @@
+// Copyright Epic Games, Inc. All Rights Reserved.
+
+#pragma once
+
+#include "CoreTypes.h"
+#include "EntitySystem/MovieSceneEntitySystem.h"
+#include "MovieSceneMixedAnimationTarget.h"
+#include "Injection/IEvaluationModifier.h"
+#include "Injection/InjectionRequest.h"
+#include "Component/AnimNextComponent.h"
+#include "MovieSceneAnimNextTargetSystem.generated.h"
+
+
+struct FAnimNextEvaluationTask;
+
+struct FMovieSceneAnimMixerEvaluationModifier : public UE::AnimNext::IEvaluationModifier
+{
+	FMovieSceneAnimMixerEvaluationModifier(TSharedPtr<FAnimNextEvaluationTask> InTaskToInject)
+		: TaskToInject(InTaskToInject)
+	{
+	}
+
+	// IEvaluationModifier interface
+	virtual void PreEvaluate(UE::AnimNext::FEvaluateTraversalContext& Context) const override {}
+
+	virtual void PostEvaluate(UE::AnimNext::FEvaluateTraversalContext& Context) const override;
+
+	TSharedPtr<FAnimNextEvaluationTask> TaskToInject;
+
+};
+
+/**
+ * Declaring a unique target for targeting an injection site on an Anim Next module. Will find the default injection site, or one specified by name.
+ */
+USTRUCT(meta=(DisplayName="Anim Next Module Injection"))
+struct MOVIESCENEANIMMIXER_API FMovieSceneAnimNextInjectionTarget : public FMovieSceneMixedAnimationTarget
+{
+	GENERATED_BODY() 
+
+	// Site name to use for injection. If empty will use the default one found for the actor.
+	// TODO: Custom UI that finds injection sites from modules.
+	UPROPERTY(EditAnywhere, Category="Animation")
+	FName InjectionSiteName;
+
+	FORCEINLINE friend uint32 GetTypeHash(const FMovieSceneAnimNextInjectionTarget& Target)
+	{
+		return HashCombine(GetTypeHash(FMovieSceneAnimNextInjectionTarget::StaticStruct()), GetTypeHash(Target.InjectionSiteName));
+	}
+};
+
+struct FMovieSceneAnimNextTargetData
+{
+	TWeakObjectPtr<UAnimNextComponent> AnimNextComponent;
+	FName InjectionSiteName;
+	UE::AnimNext::FInjectionRequestPtr InjectionRequestHandle;
+	TSharedPtr<FMovieSceneAnimMixerEvaluationModifier> Modifier;
+};
+
+// System that handles applying animation mixer evaluation tasks to an injection site in an Anim Next module
+UCLASS(MinimalAPI)
+class UMovieSceneAnimNextTargetSystem
+	: public UMovieSceneEntitySystem
+{
+public:
+
+	GENERATED_BODY()
+
+	UMovieSceneAnimNextTargetSystem(const FObjectInitializer& ObjInit);
+
+	void OnUnlink() override;
+
+	TArray<FMovieSceneAnimNextTargetData> CurrentTargets;
+
+private:
+
+	virtual void OnSchedulePersistentTasks(UE::MovieScene::IEntitySystemScheduler* TaskScheduler) override;
+};
\ No newline at end of file
