Sequencer Anim Mixer Anim Blueprint support.

The sequencer anim mixer functionality can now target animation blurprints with the new Sequencer Target anim graph node.

-Adding Sequencer Mixer Anim Subsystem, which holds a per-instance mapping of target names to anim task for coordinating between sequencer and anim bps.
-Adding Sequencer Mixer Anim Node which will search for anim tasks on the anim mixer subsystem based on a user specified name, similar to the slot system.
-Adding MovieSceneAnimBlueprintTargetSystem which adds tasks to the anim mixer subsystem of bound skeletal mesh components.
-Added MovieSceneAnimMixerEditor module to the anim mixer plugin for the editor specific classes needed to support anim BPs: i.e. AnimGraphNode_SequencerMixer and AnimBlurprintExtension_SequencerMixerTarget.

diff --git a/Engine/Plugins/Experimental/MovieSceneAnimMixer/MovieSceneAnimMixer.uplugin b/Engine/Plugins/Experimental/MovieSceneAnimMixer/MovieSceneAnimMixer.uplugin
--- a/Engine/Plugins/Experimental/MovieSceneAnimMixer/MovieSceneAnimMixer.uplugin
+++ b/Engine/Plugins/Experimental/MovieSceneAnimMixer/MovieSceneAnimMixer.uplugin
@@ -32,6 +32,11 @@
 			"Name" : "MovieSceneAnimMixer",
 			"Type" : "Runtime",
 			"LoadingPhase" : "Default"
+		},
+		{
+			"Name" : "MovieSceneAnimMixerEditor",
+			"Type": "UncookedOnly",
+			"LoadingPhase": "Default"
 		}
 	]
 }
diff --git a/Engine/Plugins/Experimental/MovieSceneAnimMixer/Source/MovieSceneAnimMixer/Private/AnimNode_SequencerMixerTarget.cpp b/Engine/Plugins/Experimental/MovieSceneAnimMixer/Source/MovieSceneAnimMixer/Private/AnimNode_SequencerMixerTarget.cpp
new file mode 100644
--- /dev/null
+++ b/Engine/Plugins/Experimental/MovieSceneAnimMixer/Source/MovieSceneAnimMixer/Private/AnimNode_SequencerMixerTarget.cpp
@@ -0,0 +1,101 @@
+﻿// Copyright Epic Games, Inc. All Rights Reserved.
+
+#include "AnimNode_SequencerMixerTarget.h"
+
+#include "AnimSubsystem_SequencerMixer.h"
+#include "DataRegistryTypes.h"
+#include "GenerationTools.h"
+#include "Animation/AnimInstanceProxy.h"
+#include "EvaluationVM/EvaluationVM.h"
+#include "Graph/AnimNext_LODPose.h"
+#include "Components/SkeletalMeshComponent.h"
+
+void FAnimNode_SequencerMixerTarget::Initialize_AnyThread(const FAnimationInitializeContext& Context)
+{
+	FAnimNode_Base::Initialize_AnyThread(Context);
+	SourcePose.Initialize(Context);
+}
+
+void FAnimNode_SequencerMixerTarget::CacheBones_AnyThread(const FAnimationCacheBonesContext& Context)
+{
+	FAnimNode_Base::CacheBones_AnyThread(Context);
+	SourcePose.CacheBones(Context);
+}
+
+void FAnimNode_SequencerMixerTarget::Update_AnyThread(const FAnimationUpdateContext& Context)
+{
+	FAnimNode_Base::Update_AnyThread(Context);
+	SourcePose.Update(Context);
+}
+
+void FAnimNode_SequencerMixerTarget::Evaluate_AnyThread(FPoseContext& Output)
+{
+	SourcePose.Evaluate(Output);
+
+	UAnimInstance* AnimInstance = Cast<UAnimInstance>(Output.GetAnimInstanceObject());
+	
+	const FAnimSubsystem_SequencerMixer& MixerSubsystem = AnimInstance->GetSubsystem<FAnimSubsystem_SequencerMixer>();
+	
+	const TSharedPtr<FAnimNextEvaluationTask>* MixerTask = MixerSubsystem.GetEvalTask(TargetName);
+	if (!MixerTask || !MixerTask->IsValid())
+	{
+		return;
+	}
+
+	USkeletalMeshComponent* SkeletalMeshComponent = AnimInstance->GetSkelMeshComponent();
+	
+	if (!SkeletalMeshComponent)
+	{
+		return;
+	}
+	
+	int32 LODLevel = Output.AnimInstanceProxy->GetLODLevel();
+
+	using namespace UE::AnimNext;
+	
+	FDataHandle RefPoseHandle = FDataRegistry::Get()->GetOrGenerateReferencePose(SkeletalMeshComponent);
+	FAnimNextGraphReferencePose GraphReferencePose(RefPoseHandle);
+
+	const UE::AnimNext::FReferencePose& RefPose = RefPoseHandle.GetRef<UE::AnimNext::FReferencePose>();
+	FAnimNextGraphLODPose ResultPose;
+	ResultPose.LODPose = FLODPoseHeap(RefPose, LODLevel, true, Output.ExpectsAdditivePose());
+
+	{
+		bool bHasValidOutput = false;
+		FEvaluationVM EvaluationVM(EEvaluationFlags::All, RefPose, LODLevel);
+
+		// Use the output pose in the mixer as the 'base pose'- push this pose first.
+		FKeyframeState Keyframe = EvaluationVM.MakeUninitializedKeyframe(false);
+		FGenerationTools::RemapPose(Output, Keyframe.Pose);
+		Keyframe.Curves.CopyFrom(Output.Curve);
+		// TODO: There is not a remap attributes the other way- do we need one?
+		Keyframe.Attributes.CopyFrom(Output.CustomAttributes);
+
+		EvaluationVM.PushValue(KEYFRAME_STACK_NAME, MakeUnique<FKeyframeState>(MoveTemp(Keyframe)));
+		
+		MixerTask->Get()->Execute(EvaluationVM);
+
+		TUniquePtr<FKeyframeState> EvaluatedKeyframe;
+		if (EvaluationVM.PopValue(KEYFRAME_STACK_NAME, EvaluatedKeyframe))
+		{
+			ResultPose.LODPose.CopyFrom(EvaluatedKeyframe->Pose);
+			ResultPose.Curves.CopyFrom(EvaluatedKeyframe->Curves);
+			ResultPose.Attributes.CopyFrom(EvaluatedKeyframe->Attributes);
+			bHasValidOutput = true;
+		}
+
+		if (!bHasValidOutput)
+		{
+			// We need to output a valid pose, generate one
+			FKeyframeState ReferenceKeyframe = EvaluationVM.MakeReferenceKeyframe(Output.ExpectsAdditivePose());
+			ResultPose.LODPose.CopyFrom(ReferenceKeyframe.Pose);
+			ResultPose.Curves.CopyFrom(ReferenceKeyframe.Curves);
+			ResultPose.Attributes.CopyFrom(ReferenceKeyframe.Attributes);
+		}
+	}
+	
+	FGenerationTools::RemapPose(ResultPose.LODPose, Output);
+	Output.Curve.CopyFrom(ResultPose.Curves);
+	FGenerationTools::RemapAttributes(ResultPose.LODPose, ResultPose.Attributes, Output);
+
+}
diff --git a/Engine/Plugins/Experimental/MovieSceneAnimMixer/Source/MovieSceneAnimMixer/Private/AnimSubsystem_SequencerMixer.cpp b/Engine/Plugins/Experimental/MovieSceneAnimMixer/Source/MovieSceneAnimMixer/Private/AnimSubsystem_SequencerMixer.cpp
new file mode 100644
--- /dev/null
+++ b/Engine/Plugins/Experimental/MovieSceneAnimMixer/Source/MovieSceneAnimMixer/Private/AnimSubsystem_SequencerMixer.cpp
@@ -0,0 +1,24 @@
+﻿// Copyright Epic Games, Inc. All Rights Reserved.
+
+#include "AnimSubsystem_SequencerMixer.h"
+
+void FAnimSubsystem_SequencerMixer::Initialize_WorkerThread()
+{
+	TargetToTaskMap.Empty();
+}
+
+void FAnimSubsystem_SequencerMixer::RegisterEvalTask(const FName TargetName,
+                                                     TSharedPtr<FAnimNextEvaluationTask> EvalTask)
+{
+	TargetToTaskMap.FindOrAdd(TargetName) = EvalTask;
+}
+
+const TSharedPtr<FAnimNextEvaluationTask>* FAnimSubsystem_SequencerMixer::GetEvalTask(const FName TargetName) const
+{
+	return TargetToTaskMap.Find(TargetName);
+}
+
+void FAnimSubsystem_SequencerMixer::ResetEvalTasks()
+{
+	TargetToTaskMap.Reset();
+}
diff --git a/Engine/Plugins/Experimental/MovieSceneAnimMixer/Source/MovieSceneAnimMixer/Private/Systems/MovieSceneAnimBlueprintTargetSystem.cpp b/Engine/Plugins/Experimental/MovieSceneAnimMixer/Source/MovieSceneAnimMixer/Private/Systems/MovieSceneAnimBlueprintTargetSystem.cpp
new file mode 100644
--- /dev/null
+++ b/Engine/Plugins/Experimental/MovieSceneAnimMixer/Source/MovieSceneAnimMixer/Private/Systems/MovieSceneAnimBlueprintTargetSystem.cpp
@@ -0,0 +1,205 @@
+﻿// Copyright Epic Games, Inc. All Rights Reserved.
+
+#include "Systems/MovieSceneAnimBlueprintTargetSystem.h"
+
+#include "AnimMixerComponentTypes.h"
+#include "AnimSubsystem_SequencerMixer.h"
+#include "SkeletalMeshRestoreState.h"
+
+#include "EntitySystem/IMovieSceneTaskScheduler.h"
+#include "EntitySystem/MovieSceneEntitySystemLinker.h"
+#include "EntitySystem/MovieSceneEntitySystemTask.h"
+#include "EntitySystem/MovieScenePreAnimatedStateSystem.h"
+#include "EntitySystem/Interrogation/MovieSceneInterrogationLinker.h"
+
+#include "Evaluation/PreAnimatedState/MovieScenePreAnimatedObjectStorage.h"
+#include "Evaluation/PreAnimatedState/MovieScenePreAnimatedStorageID.h"
+#include "Evaluation/PreAnimatedState/MovieScenePreAnimatedStorageID.inl"
+
+#include "Systems/MovieSceneAnimMixerSystem.h"
+#include "Systems/MovieSceneMixedSkeletalAnimationSystem.h"
+
+namespace UE::MovieScene
+{
+
+	/* ------- Pre-animated state for skeletal animations ------- */
+
+	struct FPreAnimatedAnimBlueprintMixerState
+	{
+		FSkeletalMeshRestoreState MeshRestoreState;
+	};
+	
+	struct FPreAnimatedAnimBlueprintMixedSkeletalAnimationTraits : FBoundObjectPreAnimatedStateTraits
+	{
+		using KeyType = FObjectKey;
+		using StorageType = FPreAnimatedAnimBlueprintMixerState;
+
+		FPreAnimatedAnimBlueprintMixerState CachePreAnimatedValue(const KeyType& Object)
+		{
+			FPreAnimatedAnimBlueprintMixerState OutCachedValue;
+
+			USkeletalMeshComponent* Component = Cast<USkeletalMeshComponent>(Object.ResolveObjectPtr());
+			if (ensure(Component))
+			{
+				OutCachedValue.MeshRestoreState.SaveState(Component);
+			}
+			
+			return OutCachedValue;
+		}
+
+		void RestorePreAnimatedValue(const KeyType& Object, StorageType& InOutCachedValue, const FRestoreStateParams& Params)
+		{
+			USkeletalMeshComponent* Component = Cast<USkeletalMeshComponent>(Object.ResolveObjectPtr());
+			if (!Component || !Component->IsRegistered())
+			{
+				return;
+			}
+			
+			FAnimSubsystem_SequencerMixer* MixerSubsystem = Component->GetAnimInstance()->FindSubsystem<FAnimSubsystem_SequencerMixer>();
+
+			if (MixerSubsystem)
+			{
+				MixerSubsystem->ResetEvalTasks();
+			}
+
+			InOutCachedValue.MeshRestoreState.RestoreLOD(Component);
+
+			// Restore pose after unbinding to force the restored pose
+			Component->SetUpdateAnimationInEditor(true);
+			Component->SetUpdateClothInEditor(true);
+			if (!Component->IsPostEvaluatingAnimation())
+			{
+				Component->TickAnimation(0.f, false);
+				Component->RefreshBoneTransforms();
+				Component->RefreshFollowerComponents();
+				Component->UpdateComponentToWorld();
+				Component->FinalizeBoneTransform();
+				Component->MarkRenderTransformDirty();
+				Component->MarkRenderDynamicDataDirty();
+			}
+
+			InOutCachedValue.MeshRestoreState.RestoreState(Component);
+
+			// if not game world, don't clean this up
+			if (Component->GetWorld() != nullptr && Component->GetWorld()->IsGameWorld() == false)
+			{
+				Component->ClearMotionVector();
+			}
+		}
+	};
+
+	
+	/** Pre-animated storage **/
+	struct FPreAnimatedBlueprintMixedSkeletalAnimationStorage : TPreAnimatedStateStorage_ObjectTraits<FPreAnimatedAnimBlueprintMixedSkeletalAnimationTraits>
+	{
+		static TAutoRegisterPreAnimatedStorageID<FPreAnimatedBlueprintMixedSkeletalAnimationStorage> StorageID;
+	};
+
+	TAutoRegisterPreAnimatedStorageID<FPreAnimatedBlueprintMixedSkeletalAnimationStorage> FPreAnimatedBlueprintMixedSkeletalAnimationStorage::StorageID;
+	
+	/* ---------------------- System task ----------------------- */
+
+	struct FEvaluateAndApplyAnimationTasksAnimBP
+	{
+	private:
+
+		UMovieSceneEntitySystemLinker* Linker;
+
+		TSharedPtr<FPreAnimatedBlueprintMixedSkeletalAnimationStorage> PreAnimatedStorage;
+
+	public:
+
+		FEvaluateAndApplyAnimationTasksAnimBP(UMovieSceneEntitySystemLinker* InLinker)
+			: Linker(InLinker)
+		{
+			PreAnimatedStorage = Linker->PreAnimatedState.GetOrCreateStorage<FPreAnimatedBlueprintMixedSkeletalAnimationStorage>();
+		}
+
+		void ForEachEntity (
+			FMovieSceneEntityID EntityID,
+			FObjectComponent MeshComponent,
+			const TInstancedStruct<FMovieSceneMixedAnimationTarget>& Target,
+			TSharedPtr<FAnimNextEvaluationTask> MixerTask
+		) const
+		{
+			// Invalid or incompatible target, return.
+			if (!Target.GetPtr<FMovieSceneAnimBlueprintTarget>())
+			{
+				return;
+			}
+
+			const FName BlueprintNodeName = Target.GetPtr<FMovieSceneAnimBlueprintTarget>()->BlueprintNodeName;
+
+			USkeletalMeshComponent* SkeletalMeshComponent = CastChecked<USkeletalMeshComponent>(MeshComponent.GetObject());
+			ensureMsgf(SkeletalMeshComponent, TEXT("Attempting to apply animation to an anim instance without a valid skeletal mesh component."));
+
+			if (!SkeletalMeshComponent || !SkeletalMeshComponent->GetSkeletalMeshAsset() || !SkeletalMeshComponent->GetAnimInstance())
+			{
+				return;
+			}
+
+			PreAnimatedStorage->BeginTrackingEntity(EntityID, true, FRootInstanceHandle(), SkeletalMeshComponent);
+			
+			PreAnimatedStorage->CachePreAnimatedValue(FCachePreAnimatedValueParams(), SkeletalMeshComponent);
+			
+			UAnimInstance* AnimInstance = SkeletalMeshComponent->GetAnimInstance();
+			
+			FAnimSubsystem_SequencerMixer* MixerSubsystem = AnimInstance->FindSubsystem<FAnimSubsystem_SequencerMixer>();
+
+			if (MixerSubsystem)
+			{
+				MixerSubsystem->RegisterEvalTask(BlueprintNodeName, MixerTask);
+			}
+			
+			if (!SkeletalMeshComponent->IsPostEvaluatingAnimation() && SkeletalMeshComponent->PoseTickedThisFrame())
+			{
+				SkeletalMeshComponent->TickAnimation(0.f, false);
+
+				SkeletalMeshComponent->RefreshBoneTransforms();
+				SkeletalMeshComponent->RefreshFollowerComponents();
+				SkeletalMeshComponent->UpdateComponentToWorld();
+				SkeletalMeshComponent->FinalizeBoneTransform();
+				SkeletalMeshComponent->MarkRenderTransformDirty();
+				SkeletalMeshComponent->MarkRenderDynamicDataDirty();
+			}
+		}
+		
+	};
+}
+
+UMovieSceneAnimBlueprintTargetSystem::UMovieSceneAnimBlueprintTargetSystem(const FObjectInitializer& ObjInit)
+	: Super(ObjInit)
+{
+	using namespace UE::MovieScene;
+	
+	FAnimMixerComponentTypes* AnimMixerComponents = FAnimMixerComponentTypes::Get();
+
+	RelevantComponent = AnimMixerComponents->MixerTask;
+	Phase = ESystemPhase::Scheduling;
+
+	SystemCategories |= FSystemInterrogator::GetExcludedFromInterrogationCategory();
+
+	if (HasAnyFlags(RF_ClassDefaultObject))
+	{
+		DefineImplicitPrerequisite(UMovieSceneAnimMixerSystem::StaticClass(), GetClass());
+		DefineImplicitPrerequisite(GetClass(), UMovieSceneRestorePreAnimatedStateSystem::StaticClass());
+	}
+}
+
+void UMovieSceneAnimBlueprintTargetSystem::OnSchedulePersistentTasks(UE::MovieScene::IEntitySystemScheduler* TaskScheduler)
+{
+	using namespace  UE::MovieScene;
+
+	FAnimMixerComponentTypes* AnimMixerComponents = FAnimMixerComponentTypes::Get();
+
+	// Evaluate anim tasks targeting animation blueprints.
+	FTaskParams Params(TEXT("Apply Animation Tasks AnimBP"));
+	Params.ForceGameThread();
+	FTaskID EvaluateTask = FEntityTaskBuilder()
+		.ReadEntityIDs()
+		.Read(AnimMixerComponents->MeshComponent)
+		.Read(AnimMixerComponents->Target)
+		.Read(AnimMixerComponents->MixerTask)
+		.SetParams(Params)
+		.Schedule_PerEntity<FEvaluateAndApplyAnimationTasksAnimBP>(&Linker->EntityManager, TaskScheduler, Linker);
+}
diff --git a/Engine/Plugins/Experimental/MovieSceneAnimMixer/Source/MovieSceneAnimMixer/Public/AnimNode_SequencerMixerTarget.h b/Engine/Plugins/Experimental/MovieSceneAnimMixer/Source/MovieSceneAnimMixer/Public/AnimNode_SequencerMixerTarget.h
new file mode 100644
--- /dev/null
+++ b/Engine/Plugins/Experimental/MovieSceneAnimMixer/Source/MovieSceneAnimMixer/Public/AnimNode_SequencerMixerTarget.h
@@ -0,0 +1,27 @@
+﻿// Copyright Epic Games, Inc. All Rights Reserved.
+
+#pragma once
+
+#include "Animation/AnimNodeBase.h"
+
+#include "AnimNode_SequencerMixerTarget.generated.h"
+
+USTRUCT(BlueprintInternalUseOnly)
+struct FAnimNode_SequencerMixerTarget final : public FAnimNode_Base
+{
+	GENERATED_BODY()
+
+	// The source input. This is passed through if there is no animation received from the level sequence, otherwise it's blended in.
+	UPROPERTY(EditAnywhere, EditFixedSize, BlueprintReadWrite, Category = Links)
+	FPoseLink SourcePose;
+
+	// The target name for the level sequence to match with when applying its animation.
+	UPROPERTY(EditAnywhere, BlueprintReadWrite, Category=Setttings, meta=(CustomizePropery))
+	FName TargetName;
+
+	// FAnimNode_Base overrides
+	MOVIESCENEANIMMIXER_API virtual void Initialize_AnyThread(const FAnimationInitializeContext& Context) override;
+	MOVIESCENEANIMMIXER_API virtual void CacheBones_AnyThread(const FAnimationCacheBonesContext& Context) override;
+	MOVIESCENEANIMMIXER_API virtual void Update_AnyThread(const FAnimationUpdateContext& Context) override;
+	MOVIESCENEANIMMIXER_API virtual void Evaluate_AnyThread(FPoseContext& Output) override;
+};
diff --git a/Engine/Plugins/Experimental/MovieSceneAnimMixer/Source/MovieSceneAnimMixer/Public/AnimSubsystem_SequencerMixer.h b/Engine/Plugins/Experimental/MovieSceneAnimMixer/Source/MovieSceneAnimMixer/Public/AnimSubsystem_SequencerMixer.h
new file mode 100644
--- /dev/null
+++ b/Engine/Plugins/Experimental/MovieSceneAnimMixer/Source/MovieSceneAnimMixer/Public/AnimSubsystem_SequencerMixer.h
@@ -0,0 +1,28 @@
+﻿// Copyright Epic Games, Inc. All Rights Reserved.
+
+#pragma once
+#include "EvaluationVM/EvaluationTask.h"
+#include "Animation/AnimSubsystemInstance.h"
+#include "Animation/AnimInstance.h"
+
+#include "AnimSubsystem_SequencerMixer.generated.h"
+
+/** Allows sequencer's anim mixer system to communicate with mixer nodes in an anim blueprint*/
+USTRUCT()
+struct FAnimSubsystem_SequencerMixer : public FAnimSubsystemInstance
+{
+	GENERATED_BODY()
+
+	// FAnimSubsystemInstance interface
+	MOVIESCENEANIMMIXER_API virtual void Initialize_WorkerThread() override;
+
+	void RegisterEvalTask(FName TargetName, TSharedPtr<FAnimNextEvaluationTask> EvalTask);
+	
+	const TSharedPtr<FAnimNextEvaluationTask>* GetEvalTask(FName TargetName) const;
+
+	void ResetEvalTasks();
+
+private:
+	TMap<FName, TSharedPtr<FAnimNextEvaluationTask>> TargetToTaskMap;
+	
+};
diff --git a/Engine/Plugins/Experimental/MovieSceneAnimMixer/Source/MovieSceneAnimMixer/Public/Systems/MovieSceneAnimBlueprintTargetSystem.h b/Engine/Plugins/Experimental/MovieSceneAnimMixer/Source/MovieSceneAnimMixer/Public/Systems/MovieSceneAnimBlueprintTargetSystem.h
new file mode 100644
--- /dev/null
+++ b/Engine/Plugins/Experimental/MovieSceneAnimMixer/Source/MovieSceneAnimMixer/Public/Systems/MovieSceneAnimBlueprintTargetSystem.h
@@ -0,0 +1,41 @@
+﻿// Copyright Epic Games, Inc. All Rights Reserved.
+
+#pragma once
+#include "MovieSceneMixedAnimationTarget.h"
+#include "EntitySystem/MovieSceneEntitySystem.h"
+#include "Animation/AnimInstance.h"
+
+#include "MovieSceneAnimBlueprintTargetSystem.generated.h"
+
+/**
+ * Declaring a unique target for the "Sequencer Mixer Target" node in an Anim BP. Will match the node's name with the name specified here.
+ */
+USTRUCT(meta=(DisplayName="Anim Blueprint Target"))
+struct MOVIESCENEANIMMIXER_API FMovieSceneAnimBlueprintTarget : public FMovieSceneMixedAnimationTarget
+{
+	GENERATED_BODY() 
+
+	// Node name to use for injection.
+	UPROPERTY(EditAnywhere, Category="Animation")
+	FName BlueprintNodeName;
+
+	FORCEINLINE friend uint32 GetTypeHash(const FMovieSceneAnimBlueprintTarget& Target)
+	{
+		return HashCombine(GetTypeHash(FMovieSceneAnimBlueprintTarget::StaticStruct()), GetTypeHash(Target.BlueprintNodeName));
+	}
+};
+
+UCLASS(MinimalAPI)
+class UMovieSceneAnimBlueprintTargetSystem : public UMovieSceneEntitySystem
+{
+public:
+	
+	GENERATED_BODY()
+
+	UMovieSceneAnimBlueprintTargetSystem(const FObjectInitializer& ObjInit);
+
+private:
+
+	virtual void OnSchedulePersistentTasks(UE::MovieScene::IEntitySystemScheduler* TaskScheduler) override;
+	
+};
diff --git a/Engine/Plugins/Experimental/MovieSceneAnimMixer/Source/MovieSceneAnimMixerEditor/MovieSceneAnimMixerEditor.Build.cs b/Engine/Plugins/Experimental/MovieSceneAnimMixer/Source/MovieSceneAnimMixerEditor/MovieSceneAnimMixerEditor.Build.cs
new file mode 100644
--- /dev/null
+++ b/Engine/Plugins/Experimental/MovieSceneAnimMixer/Source/MovieSceneAnimMixerEditor/MovieSceneAnimMixerEditor.Build.cs
@@ -0,0 +1,31 @@
+﻿// Copyright Epic Games, Inc. All Rights Reserved.
+
+namespace UnrealBuildTool.Rules
+{
+	public class MovieSceneAnimMixerEditor : ModuleRules
+	{
+		public MovieSceneAnimMixerEditor(ReadOnlyTargetRules Target) : base(Target)
+		{
+			PublicDependencyModuleNames.AddRange(
+				new string[]
+				{
+					"Core",
+					"CoreUObject",
+					"Engine",
+					"MovieScene",
+					"AnimGraph",
+					"AnimationCore",
+					"AnimGraphRuntime", 
+					"MovieSceneAnimMixer",
+				}
+			);
+			
+			PrivateDependencyModuleNames.AddRange(
+				new string[]
+				{
+					"BlueprintGraph"
+				}
+			);
+		}
+	}
+}
\ No newline at end of file
diff --git a/Engine/Plugins/Experimental/MovieSceneAnimMixer/Source/MovieSceneAnimMixerEditor/Private/AnimGraphNode_SequencerMixerTarget.cpp b/Engine/Plugins/Experimental/MovieSceneAnimMixer/Source/MovieSceneAnimMixerEditor/Private/AnimGraphNode_SequencerMixerTarget.cpp
new file mode 100644
--- /dev/null
+++ b/Engine/Plugins/Experimental/MovieSceneAnimMixer/Source/MovieSceneAnimMixerEditor/Private/AnimGraphNode_SequencerMixerTarget.cpp
@@ -0,0 +1,47 @@
+﻿// Copyright Epic Games, Inc. All Rights Reserved.
+
+#include "AnimGraphNode_SequencerMixerTarget.h"
+
+#include "AnimBlueprintExtension.h"
+#include "AnimBlueprintExtension_SequencerMixerTarget.h"
+
+#define LOCTEXT_NAMESPACE "SequencerMixerNodes"
+
+FLinearColor UAnimGraphNode_SequencerMixerTarget::GetNodeTitleColor() const
+{
+	return Super::GetNodeTitleColor();
+}
+
+FText UAnimGraphNode_SequencerMixerTarget::GetTooltipText() const
+{
+	FFormatNamedArguments Args;
+	Args.Add(TEXT("TargetName"), FText::FromName(Node.TargetName));
+	
+	return FText::Format(LOCTEXT("TargetNodeToolTip", "Receives Animation Mixer results from animation sections in a Level Sequence. Target Name: '{TargetName}'"), Args);
+}
+
+FText UAnimGraphNode_SequencerMixerTarget::GetNodeTitle(ENodeTitleType::Type TitleType) const
+{
+	FFormatNamedArguments Args;
+	Args.Add(TEXT("TargetName"), FText::FromName(Node.TargetName));
+	
+	return FText::Format(LOCTEXT("TargetNodeTitle", "Sequencer Target: '{TargetName}'"), Args);
+}
+
+FString UAnimGraphNode_SequencerMixerTarget::GetNodeCategory() const
+{
+	return TEXT("Sequencer|Animation");
+}
+
+void UAnimGraphNode_SequencerMixerTarget::BakeDataDuringCompilation(FCompilerResultsLog& MessageLog)
+{
+	Super::BakeDataDuringCompilation(MessageLog);
+}
+
+void UAnimGraphNode_SequencerMixerTarget::GetRequiredExtensions(TArray<TSubclassOf<UAnimBlueprintExtension>>& OutExtensions) const
+{
+	Super::GetRequiredExtensions(OutExtensions);
+	OutExtensions.Add(UAnimBlueprintExtension_SequencerMixerTarget::StaticClass());
+}
+
+#undef LOCTEXT_NAMESPACE
diff --git a/Engine/Plugins/Experimental/MovieSceneAnimMixer/Source/MovieSceneAnimMixerEditor/Private/MovieSceneAnimMixerEditorModule.cpp b/Engine/Plugins/Experimental/MovieSceneAnimMixer/Source/MovieSceneAnimMixerEditor/Private/MovieSceneAnimMixerEditorModule.cpp
new file mode 100644
--- /dev/null
+++ b/Engine/Plugins/Experimental/MovieSceneAnimMixer/Source/MovieSceneAnimMixerEditor/Private/MovieSceneAnimMixerEditorModule.cpp
@@ -0,0 +1,21 @@
+﻿// Copyright Epic Games, Inc. All Rights Reserved.
+
+#include "MovieSceneAnimMixerEditorModule.h"
+
+#include "Modules/ModuleInterface.h"
+#include "Modules/ModuleManager.h"
+
+namespace UE::MovieScene
+{
+	void FMovieSceneAnimMixerEditorModule::StartupModule()
+	{
+		
+	}
+
+	void FMovieSceneAnimMixerEditorModule::ShutdownModule()
+	{
+		
+	}
+}
+
+IMPLEMENT_MODULE(UE::MovieScene::FMovieSceneAnimMixerEditorModule, MovieSceneAnimMixerEditor)
\ No newline at end of file
diff --git a/Engine/Plugins/Experimental/MovieSceneAnimMixer/Source/MovieSceneAnimMixerEditor/Private/MovieSceneAnimMixerEditorModule.h b/Engine/Plugins/Experimental/MovieSceneAnimMixer/Source/MovieSceneAnimMixerEditor/Private/MovieSceneAnimMixerEditorModule.h
new file mode 100644
--- /dev/null
+++ b/Engine/Plugins/Experimental/MovieSceneAnimMixer/Source/MovieSceneAnimMixerEditor/Private/MovieSceneAnimMixerEditorModule.h
@@ -0,0 +1,20 @@
+﻿// Copyright Epic Games, Inc. All Rights Reserved.
+
+#pragma once
+
+#include "CoreMinimal.h"
+#include "Modules/ModuleInterface.h"
+
+namespace UE::MovieScene
+{
+	class FMovieSceneAnimMixerEditorModule : public IModuleInterface
+	{
+	public:
+
+		// IModuleInterface interface
+		virtual void StartupModule() override;
+		virtual void ShutdownModule() override;
+	};
+}
+
+
diff --git a/Engine/Plugins/Experimental/MovieSceneAnimMixer/Source/MovieSceneAnimMixerEditor/Public/AnimBlueprintExtension_SequencerMixerTarget.h b/Engine/Plugins/Experimental/MovieSceneAnimMixer/Source/MovieSceneAnimMixerEditor/Public/AnimBlueprintExtension_SequencerMixerTarget.h
new file mode 100644
--- /dev/null
+++ b/Engine/Plugins/Experimental/MovieSceneAnimMixer/Source/MovieSceneAnimMixerEditor/Public/AnimBlueprintExtension_SequencerMixerTarget.h
@@ -0,0 +1,21 @@
+﻿// Copyright Epic Games, Inc. All Rights Reserved.
+
+#pragma once
+
+#include "CoreMinimal.h"
+#include "AnimSubsystem_SequencerMixer.h"
+#include "AnimBlueprintExtension.h"
+#include "AnimBlueprintExtension_SequencerMixerTarget.generated.h"
+
+/**
+ * Extension for brokering Animation Mixer data between sequencer and animation blueprints.
+ */
+UCLASS()
+class MOVIESCENEANIMMIXEREDITOR_API UAnimBlueprintExtension_SequencerMixerTarget : public UAnimBlueprintExtension
+{
+	GENERATED_BODY()
+
+private:
+	UPROPERTY()
+	FAnimSubsystem_SequencerMixer Subsystem;
+};
diff --git a/Engine/Plugins/Experimental/MovieSceneAnimMixer/Source/MovieSceneAnimMixerEditor/Public/AnimGraphNode_SequencerMixerTarget.h b/Engine/Plugins/Experimental/MovieSceneAnimMixer/Source/MovieSceneAnimMixerEditor/Public/AnimGraphNode_SequencerMixerTarget.h
new file mode 100644
--- /dev/null
+++ b/Engine/Plugins/Experimental/MovieSceneAnimMixer/Source/MovieSceneAnimMixerEditor/Public/AnimGraphNode_SequencerMixerTarget.h
@@ -0,0 +1,32 @@
+﻿// Copyright Epic Games, Inc. All Rights Reserved.
+
+#pragma once
+
+#include "AnimGraphNode_Base.h"
+#include "AnimNode_SequencerMixerTarget.h"
+
+#include "AnimGraphNode_SequencerMixerTarget.generated.h"
+
+UCLASS(MinimalAPI)
+class UAnimGraphNode_SequencerMixerTarget : public UAnimGraphNode_Base
+{
+public:
+
+	GENERATED_BODY()
+
+	UPROPERTY(EditAnywhere, Category=Settings)
+	FAnimNode_SequencerMixerTarget Node;
+
+	//~ Begin UEdGraphNode Interface.
+	MOVIESCENEANIMMIXEREDITOR_API virtual FLinearColor GetNodeTitleColor() const override;
+	MOVIESCENEANIMMIXEREDITOR_API virtual FText GetTooltipText() const override;
+	MOVIESCENEANIMMIXEREDITOR_API virtual FText GetNodeTitle(ENodeTitleType::Type TitleType) const override;
+	//~ End UEdGraphNode Interface.
+
+	//~ Begin UAnimGraphNode_Base Interface
+	MOVIESCENEANIMMIXEREDITOR_API virtual FString GetNodeCategory() const override;
+	MOVIESCENEANIMMIXEREDITOR_API virtual void BakeDataDuringCompilation(class FCompilerResultsLog& MessageLog) override;
+	MOVIESCENEANIMMIXEREDITOR_API virtual void GetRequiredExtensions(TArray<TSubclassOf<UAnimBlueprintExtension>>& OutExtensions) const override;
+	//~ End UAnimGraphNode_Base Interface
+	
+};
