Create default name for Sequencer Mixer Target nodes

diff --git a/Engine/Plugins/Experimental/MovieSceneAnimMixer/Source/MovieSceneAnimMixer/Private/AnimNode_SequencerMixerTarget.cpp b/Engine/Plugins/Experimental/MovieSceneAnimMixer/Source/MovieSceneAnimMixer/Private/AnimNode_SequencerMixerTarget.cpp
--- a/Engine/Plugins/Experimental/MovieSceneAnimMixer/Source/MovieSceneAnimMixer/Private/AnimNode_SequencerMixerTarget.cpp
+++ b/Engine/Plugins/Experimental/MovieSceneAnimMixer/Source/MovieSceneAnimMixer/Private/AnimNode_SequencerMixerTarget.cpp
@@ -10,6 +10,14 @@
 #include "Graph/AnimNext_LODPose.h"
 #include "Components/SkeletalMeshComponent.h"
 
+
+const FName FAnimNode_SequencerMixerTarget::DefaultTargetName = FName(TEXT("DefaultTarget"));
+
+FAnimNode_SequencerMixerTarget::FAnimNode_SequencerMixerTarget()
+	: TargetName(DefaultTargetName)
+{
+}
+
 void FAnimNode_SequencerMixerTarget::Initialize_AnyThread(const FAnimationInitializeContext& Context)
 {
 	FAnimNode_Base::Initialize_AnyThread(Context);
diff --git a/Engine/Plugins/Experimental/MovieSceneAnimMixer/Source/MovieSceneAnimMixer/Public/AnimNode_SequencerMixerTarget.h b/Engine/Plugins/Experimental/MovieSceneAnimMixer/Source/MovieSceneAnimMixer/Public/AnimNode_SequencerMixerTarget.h
--- a/Engine/Plugins/Experimental/MovieSceneAnimMixer/Source/MovieSceneAnimMixer/Public/AnimNode_SequencerMixerTarget.h
+++ b/Engine/Plugins/Experimental/MovieSceneAnimMixer/Source/MovieSceneAnimMixer/Public/AnimNode_SequencerMixerTarget.h
@@ -11,6 +11,8 @@ struct FAnimNode_SequencerMixerTarget final : public FAnimNode_Base
 {
 	GENERATED_BODY()
 
+	static MOVIESCENEANIMMIXER_API const FName DefaultTargetName;
+
 	// The source input. This is passed through if there is no animation received from the level sequence, otherwise it's blended in.
 	UPROPERTY(EditAnywhere, EditFixedSize, BlueprintReadWrite, Category = Links)
 	FPoseLink SourcePose;
@@ -20,6 +22,7 @@ struct FAnimNode_SequencerMixerTarget final : public FAnimNode_Base
 	FName TargetName;
 
 	// FAnimNode_Base overrides
+	MOVIESCENEANIMMIXER_API FAnimNode_SequencerMixerTarget();
 	MOVIESCENEANIMMIXER_API virtual void Initialize_AnyThread(const FAnimationInitializeContext& Context) override;
 	MOVIESCENEANIMMIXER_API virtual void CacheBones_AnyThread(const FAnimationCacheBonesContext& Context) override;
 	MOVIESCENEANIMMIXER_API virtual void Update_AnyThread(const FAnimationUpdateContext& Context) override;
diff --git a/Engine/Plugins/Experimental/MovieSceneAnimMixer/Source/MovieSceneAnimMixer/Public/Systems/MovieSceneAnimBlueprintTargetSystem.h b/Engine/Plugins/Experimental/MovieSceneAnimMixer/Source/MovieSceneAnimMixer/Public/Systems/MovieSceneAnimBlueprintTargetSystem.h
--- a/Engine/Plugins/Experimental/MovieSceneAnimMixer/Source/MovieSceneAnimMixer/Public/Systems/MovieSceneAnimBlueprintTargetSystem.h
+++ b/Engine/Plugins/Experimental/MovieSceneAnimMixer/Source/MovieSceneAnimMixer/Public/Systems/MovieSceneAnimBlueprintTargetSystem.h
@@ -1,6 +1,7 @@
 ﻿// Copyright Epic Games, Inc. All Rights Reserved.
 
 #pragma once
+#include "AnimNode_SequencerMixerTarget.h"
 #include "MovieSceneMixedAnimationTarget.h"
 #include "EntitySystem/MovieSceneEntitySystem.h"
 #include "Animation/AnimInstance.h"
@@ -23,6 +24,12 @@ struct MOVIESCENEANIMMIXER_API FMovieSceneAnimBlueprintTarget : public FMovieSce
 	{
 		return HashCombine(GetTypeHash(FMovieSceneAnimBlueprintTarget::StaticStruct()), GetTypeHash(Target.BlueprintNodeName));
 	}
+
+	FMovieSceneAnimBlueprintTarget()
+		: BlueprintNodeName(FAnimNode_SequencerMixerTarget::DefaultTargetName)
+	{
+		
+	}
 };
 
 UCLASS(MinimalAPI)
