Sequencer: Automatically select appropriate default animaton target if none selected.

diff --git a/Engine/Plugins/Experimental/MovieSceneAnimMixer/Source/MovieSceneAnimMixer/Private/Systems/MovieSceneAnimInstanceTargetSystem.cpp b/Engine/Plugins/Experimental/MovieSceneAnimMixer/Source/MovieSceneAnimMixer/Private/Systems/MovieSceneAnimInstanceTargetSystem.cpp
--- a/Engine/Plugins/Experimental/MovieSceneAnimMixer/Source/MovieSceneAnimMixer/Private/Systems/MovieSceneAnimInstanceTargetSystem.cpp
+++ b/Engine/Plugins/Experimental/MovieSceneAnimMixer/Source/MovieSceneAnimMixer/Private/Systems/MovieSceneAnimInstanceTargetSystem.cpp
@@ -418,8 +418,7 @@ namespace UE::MovieScene
 				return;
 			}
 
-			USkeletalMeshComponent* SkeletalMeshComponent = CastChecked<USkeletalMeshComponent>(MeshComponent.GetObject());
-			ensureMsgf(SkeletalMeshComponent, TEXT("Attempting to apply animation to an anim instance without a valid skeletal mesh component."));
+			USkeletalMeshComponent* SkeletalMeshComponent = Cast<USkeletalMeshComponent>(MeshComponent.GetObject());
 
 			if (!SkeletalMeshComponent || !SkeletalMeshComponent->GetSkeletalMeshAsset())
 			{
diff --git a/Engine/Plugins/Experimental/MovieSceneAnimMixer/Source/MovieSceneAnimMixer/Private/Systems/MovieSceneAnimMixerSystem.cpp b/Engine/Plugins/Experimental/MovieSceneAnimMixer/Source/MovieSceneAnimMixer/Private/Systems/MovieSceneAnimMixerSystem.cpp
--- a/Engine/Plugins/Experimental/MovieSceneAnimMixer/Source/MovieSceneAnimMixer/Private/Systems/MovieSceneAnimMixerSystem.cpp
+++ b/Engine/Plugins/Experimental/MovieSceneAnimMixer/Source/MovieSceneAnimMixer/Private/Systems/MovieSceneAnimMixerSystem.cpp
@@ -28,6 +28,11 @@
 #include "EntitySystem/MovieSceneEntityGroupingSystem.h"
 #include "Components/SkeletalMeshComponent.h"
 #include "GameFramework/Actor.h"
+#include "Component/AnimNextComponent.h"
+#include "Systems/MovieSceneAnimInstanceTargetSystem.h"
+#include "Systems/MovieSceneAnimNextTargetSystem.h"
+#include "Systems/MovieSceneAnimBlueprintTargetSystem.h"
+#include "AnimSubsystem_SequencerMixer.h"
 
 #include UE_INLINE_GENERATED_CPP_BY_NAME(MovieSceneAnimMixerSystem)
 
@@ -94,7 +99,7 @@ void UMovieSceneAnimMixerSystem::OnRun(FSystemTaskPrerequisites& InPrerequisites
 		FEntityTaskBuilder()
 		.ReadEntityIDs()
 		.Read(BuiltInComponents->BoundObjectKey)
-		.Read(AnimMixerComponents->Target)
+		.Write(AnimMixerComponents->Target)
 		.Read(AnimMixerComponents->Task)
 		.Read(AnimMixerComponents->Priority)
 		.ReadOptional(AnimMixerComponents->RootMotionType)
@@ -104,11 +109,13 @@ void UMovieSceneAnimMixerSystem::OnRun(FSystemTaskPrerequisites& InPrerequisites
 		(
 			FMovieSceneEntityID EntityID,
 			FObjectKey BoundObjectKey,
-			const TInstancedStruct<FMovieSceneMixedAnimationTarget>& Target,
+			TInstancedStruct<FMovieSceneMixedAnimationTarget>& Target,
 			TSharedPtr<FAnimNextEvaluationTask> Task,
 			int32 Priority,
 			const ESwapRootBone* RootMotionType)
 			{
+				// For new entities, we 'resolve' the animation target so if 'Automatic' is picked we choose the right one automatically.
+				Target = ResolveAnimationTarget(BoundObjectKey, Target);
 				FMovieSceneAnimMixerKey Key({ BoundObjectKey, Target });
 				FMovieSceneAnimMixer* Mixer = Mixers.Find(Key);
 				if (!Mixer)
@@ -382,4 +389,50 @@ void UMovieSceneAnimMixerSystem::OnCleanTaggedGarbage()
 			{
 				Mixers.Remove(FMovieSceneAnimMixerKey({ MeshComponent.GetObject(), Target }));
 			});
+}
+
+const TInstancedStruct<FMovieSceneMixedAnimationTarget> UMovieSceneAnimMixerSystem::ResolveAnimationTarget(FObjectKey ObjectKey, const TInstancedStruct<FMovieSceneMixedAnimationTarget> InTarget)
+{
+	// If user has selected the default 'automatic' target, attempt to choose one automatically for them.
+	if (InTarget.GetScriptStruct() == FMovieSceneMixedAnimationTarget::StaticStruct())
+	{
+		if (UObject* Object = ObjectKey.ResolveObjectPtr())
+		{
+			USkeletalMeshComponent* SkeletalMeshComponent = Cast<USkeletalMeshComponent>(Object);
+			UAnimNextComponent* AnimNextComponent = nullptr;
+			if (!SkeletalMeshComponent)
+			{
+				if (AActor* Actor = Cast<AActor>(Object))
+				{
+					SkeletalMeshComponent = Actor->FindComponentByClass<USkeletalMeshComponent>();
+					AnimNextComponent = SkeletalMeshComponent->GetOwner()->FindComponentByClass<UAnimNextComponent>();
+				}
+			}
+			else 
+			{
+				AnimNextComponent = SkeletalMeshComponent->GetOwner()->FindComponentByClass<UAnimNextComponent>();
+			}
+
+			if (AnimNextComponent != nullptr && (!SkeletalMeshComponent || !SkeletalMeshComponent->bEnableAnimation))
+			{
+				// If we have an anim next component and the skeletal mesh component has animation disabled, default to anim next target
+				return TInstancedStruct<FMovieSceneAnimNextInjectionTarget>::Make();
+			}
+			else if (SkeletalMeshComponent)
+			{
+				if (UAnimInstance* AnimInstance = SkeletalMeshComponent->GetAnimInstance())
+				{
+					if (const FAnimSubsystem_SequencerMixer* MixerSubsystem = AnimInstance->FindSubsystem<FAnimSubsystem_SequencerMixer>())
+					{
+						// We have an anim blueprint with sequencer mixer target node(s). Create a target using the default target name.
+						return TInstancedStruct<FMovieSceneAnimBlueprintTarget>::Make();
+					}
+				}
+
+				// Fallback to using a custom anim instance as the target
+				return TInstancedStruct<FMovieSceneAnimInstanceTarget>::Make();
+			}
+		}
+	}
+	return InTarget;
 }
\ No newline at end of file
diff --git a/Engine/Plugins/Experimental/MovieSceneAnimMixer/Source/MovieSceneAnimMixer/Private/Systems/MovieSceneMixedSkeletalAnimationSystem.cpp b/Engine/Plugins/Experimental/MovieSceneAnimMixer/Source/MovieSceneAnimMixer/Private/Systems/MovieSceneMixedSkeletalAnimationSystem.cpp
--- a/Engine/Plugins/Experimental/MovieSceneAnimMixer/Source/MovieSceneAnimMixer/Private/Systems/MovieSceneMixedSkeletalAnimationSystem.cpp
+++ b/Engine/Plugins/Experimental/MovieSceneAnimMixer/Source/MovieSceneAnimMixer/Private/Systems/MovieSceneMixedSkeletalAnimationSystem.cpp
@@ -150,7 +150,8 @@ struct FGatherMixableSkeletalAnimations
 
 			// Get the bound skeletal mesh component.
 			USkeletalMeshComponent* SkeletalMeshComponent = CastChecked<USkeletalMeshComponent>(BoundObject);
-			if (!SkeletalMeshComponent || AnimParams.Animation == nullptr)
+			UAnimSequence* AnimSequence = Cast<UAnimSequence>(AnimParams.Animation);
+			if (!SkeletalMeshComponent || AnimSequence == nullptr)
 			{
 				continue;
 			}
@@ -195,14 +196,14 @@ struct FGatherMixableSkeletalAnimations
 			}
 
 			FMixedAnimSkeletalAnimationData AnimData;
-			AnimData.AnimSequence = Cast<UAnimSequence>(AnimParams.Animation);
+			AnimData.AnimSequence = AnimSequence;
 			AnimData.RootMotionOverride = RootMotionOverride;
 			AnimData.FromPosition = PreviousEvalTime;
 			AnimData.ToPosition = EvalTime;
 			AnimData.bFireNotifies = bFireNotifies && !AnimParams.bSkipAnimNotifiers && !bLooped;
 			AnimData.SwapRootBone = AnimParams.SwapRootBone;
 			// TODO: We want to allow for additive to be set on the section properties
-			AnimData.bAdditive = AnimData.AnimSequence->IsValidAdditive();
+			AnimData.bAdditive = AnimData.AnimSequence.IsValid() ? AnimData.AnimSequence->IsValidAdditive() : false;
 
 			// TODO: Need to create a mirroring task and add that if appropriate
 			// TODO: Notifies- perhaps this should be a part of the animnext side of things?
diff --git a/Engine/Plugins/Experimental/MovieSceneAnimMixer/Source/MovieSceneAnimMixer/Public/Systems/MovieSceneAnimMixerSystem.h b/Engine/Plugins/Experimental/MovieSceneAnimMixer/Source/MovieSceneAnimMixer/Public/Systems/MovieSceneAnimMixerSystem.h
--- a/Engine/Plugins/Experimental/MovieSceneAnimMixer/Source/MovieSceneAnimMixer/Public/Systems/MovieSceneAnimMixerSystem.h
+++ b/Engine/Plugins/Experimental/MovieSceneAnimMixer/Source/MovieSceneAnimMixer/Public/Systems/MovieSceneAnimMixerSystem.h
@@ -89,10 +89,14 @@ public:
 	// Map of animation mixers
 	TMap<FMovieSceneAnimMixerKey, FMovieSceneAnimMixer> Mixers;
 
+
 private:
 
 	virtual void OnRun(FSystemTaskPrerequisites& InPrerequisites, FSystemSubsequentTasks& Subsequents) override;
 	virtual void OnSchedulePersistentTasks(UE::MovieScene::IEntitySystemScheduler* TaskScheduler) override;
-	virtual void OnCleanTaggedGarbage() override;
+	virtual void OnCleanTaggedGarbage() override; 
+	
+	const TInstancedStruct<FMovieSceneMixedAnimationTarget> ResolveAnimationTarget(FObjectKey ObjectKey, const TInstancedStruct<FMovieSceneMixedAnimationTarget> InTarget);
 
+	
 };
\ No newline at end of file
diff --git a/Engine/Source/Runtime/MovieSceneTracks/Private/Sections/MovieSceneSkeletalAnimationSection.cpp b/Engine/Source/Runtime/MovieSceneTracks/Private/Sections/MovieSceneSkeletalAnimationSection.cpp
--- a/Engine/Source/Runtime/MovieSceneTracks/Private/Sections/MovieSceneSkeletalAnimationSection.cpp
+++ b/Engine/Source/Runtime/MovieSceneTracks/Private/Sections/MovieSceneSkeletalAnimationSection.cpp
@@ -121,6 +121,8 @@ UMovieSceneSkeletalAnimationSection::UMovieSceneSkeletalAnimationSection( const
 	bMatchRotationPitch = false;
 	bMatchIncludeZHeight = false;
 	bDebugForceTickPose = false;
+
+	MixedAnimationTarget = TInstancedStruct<FMovieSceneMixedAnimationTarget>::Make();
 }
 
 EMovieSceneChannelProxyType UMovieSceneSkeletalAnimationSection::CacheChannelProxy()
