[MovieSceneAnimMixer] Replace CastChecked with Cast as the SKM can be null.

There's in fact an ensure right after the cast.

diff --git a/Engine/Plugins/Experimental/MovieSceneAnimMixer/Source/MovieSceneAnimMixer/Private/Systems/MovieSceneRootMotionSystem.cpp b/Engine/Plugins/Experimental/MovieSceneAnimMixer/Source/MovieSceneAnimMixer/Private/Systems/MovieSceneRootMotionSystem.cpp
--- a/Engine/Plugins/Experimental/MovieSceneAnimMixer/Source/MovieSceneAnimMixer/Private/Systems/MovieSceneRootMotionSystem.cpp
+++ b/Engine/Plugins/Experimental/MovieSceneAnimMixer/Source/MovieSceneAnimMixer/Private/Systems/MovieSceneRootMotionSystem.cpp
@@ -67,7 +67,7 @@ namespace UE::MovieScene
 
 			FAnimMixerComponentTypes* AnimMixerTypes = FAnimMixerComponentTypes::Get();
 
-			USkeletalMeshComponent* SkeletalMeshComponent = CastChecked<USkeletalMeshComponent>(MeshComponent.GetObject());
+			USkeletalMeshComponent* SkeletalMeshComponent = Cast<USkeletalMeshComponent>(MeshComponent.GetObject());
 			ensureMsgf(SkeletalMeshComponent, TEXT("Attempting to apply animation to an anim instance without a valid skeletal mesh component."));
 
 			if (!SkeletalMeshComponent || !SkeletalMeshComponent->GetSkeletalMeshAsset())
