Sequencer: Refactor root motion handling in the anim mixer plugin to allow root motion both in 'world space' and 'animation space'. This paves the way for the 'stitch track' and other animation tracks that produce root motion in world space.

diff --git a/Engine/Plugins/Experimental/MovieSceneAnimMixer/Source/MovieSceneAnimMixer/MovieSceneAnimMixer.Build.cs b/Engine/Plugins/Experimental/MovieSceneAnimMixer/Source/MovieSceneAnimMixer/MovieSceneAnimMixer.Build.cs
--- a/Engine/Plugins/Experimental/MovieSceneAnimMixer/Source/MovieSceneAnimMixer/MovieSceneAnimMixer.Build.cs
+++ b/Engine/Plugins/Experimental/MovieSceneAnimMixer/Source/MovieSceneAnimMixer/MovieSceneAnimMixer.Build.cs
@@ -26,6 +26,15 @@ namespace UnrealBuildTool.Rules
 					"MovieSceneTracks",
 				}
 			);
+
+			if (Target.bBuildWithEditorOnlyData && Target.bBuildEditor)
+			{
+				PrivateDependencyModuleNames.AddRange(new string[]
+					{
+						"EditorFramework",
+						"UnrealEd"
+					});
+			}
 		}
 	}
 }
\ No newline at end of file
diff --git a/Engine/Plugins/Experimental/MovieSceneAnimMixer/Source/MovieSceneAnimMixer/Private/AnimMixerComponentTypes.cpp b/Engine/Plugins/Experimental/MovieSceneAnimMixer/Source/MovieSceneAnimMixer/Private/AnimMixerComponentTypes.cpp
--- a/Engine/Plugins/Experimental/MovieSceneAnimMixer/Source/MovieSceneAnimMixer/Private/AnimMixerComponentTypes.cpp
+++ b/Engine/Plugins/Experimental/MovieSceneAnimMixer/Source/MovieSceneAnimMixer/Private/AnimMixerComponentTypes.cpp
@@ -40,8 +40,10 @@ namespace UE::MovieScene
 		ComponentRegistry->NewComponentType(&Target, TEXT("Mixed Animation Target"));
 		ComponentRegistry->NewComponentType(&Task, TEXT("Mixed Animation Task"));
 		ComponentRegistry->NewComponentType(&MixerTask, TEXT("Mixed Animation Mixer Task"));
-		ComponentRegistry->NewComponentType(&RootMotionType, TEXT("Root Motion Type"));
+		ComponentRegistry->NewComponentType(&RootMotionSettings, TEXT("Root Motion Settings"));
 		ComponentRegistry->NewComponentType(&MeshComponent, TEXT("Mixed Animation Mesh Component"));
+		ComponentRegistry->NewComponentType(&MixerRootMotion, TEXT("Root Motion"));
+
 
 		Tags.RequiresBlending = ComponentRegistry->NewTag(TEXT("Requires Blending"));
 
@@ -49,8 +51,9 @@ namespace UE::MovieScene
 		ComponentRegistry->Factories.DuplicateChildComponent(Target);
 		ComponentRegistry->Factories.DuplicateChildComponent(Task);
 		ComponentRegistry->Factories.DuplicateChildComponent(MixerTask);
-		ComponentRegistry->Factories.DuplicateChildComponent(RootMotionType);
+		ComponentRegistry->Factories.DuplicateChildComponent(RootMotionSettings);
 		ComponentRegistry->Factories.DuplicateChildComponent(MeshComponent);
+		ComponentRegistry->Factories.DuplicateChildComponent(MixerRootMotion);
 
 		ComponentRegistry->Factories.DefineChildComponent(Tags.RequiresBlending, Tags.RequiresBlending);
 
diff --git a/Engine/Plugins/Experimental/MovieSceneAnimMixer/Source/MovieSceneAnimMixer/Private/MovieSceneAnimMixerModule.h b/Engine/Plugins/Experimental/MovieSceneAnimMixer/Source/MovieSceneAnimMixer/Private/MovieSceneAnimMixerModule.h
deleted file mode 100644
diff --git a/Engine/Plugins/Experimental/MovieSceneAnimMixer/Source/MovieSceneAnimMixer/Private/Systems/MovieSceneAnimMixerSystem.cpp b/Engine/Plugins/Experimental/MovieSceneAnimMixer/Source/MovieSceneAnimMixer/Private/Systems/MovieSceneAnimMixerSystem.cpp
--- a/Engine/Plugins/Experimental/MovieSceneAnimMixer/Source/MovieSceneAnimMixer/Private/Systems/MovieSceneAnimMixerSystem.cpp
+++ b/Engine/Plugins/Experimental/MovieSceneAnimMixer/Source/MovieSceneAnimMixer/Private/Systems/MovieSceneAnimMixerSystem.cpp
@@ -36,6 +36,109 @@
 
 #include UE_INLINE_GENERATED_CPP_BY_NAME(MovieSceneAnimMixerSystem)
 
+
+namespace UE::MovieScene
+{
+	struct FUpdateTaskEntities
+	{
+		FUpdateTaskEntities(UMovieSceneAnimMixerSystem* InSystem)
+			:System(InSystem)
+
+		{
+
+		}
+		void ForEachEntity(FMovieSceneEntityID EntityID,
+			FObjectKey BoundObjectKey,
+			const TInstancedStruct<FMovieSceneMixedAnimationTarget>& Target,
+			TSharedPtr<FAnimNextEvaluationTask> Task,
+			const double* WeightAndEasings) const
+		{
+			FMovieSceneAnimMixerKey Key({ BoundObjectKey, Target });
+			if (FMovieSceneAnimMixer* Mixer = System->Mixers.Find(Key))
+			{
+				// Find the entry in question
+				if (FMovieSceneAnimMixerEntry* Entry = Mixer->Entries.FindByPredicate([EntityID](const FMovieSceneAnimMixerEntry& TestEntry) { return TestEntry.EntityID == EntityID; }))
+				{
+					Entry->EvalTask = Task;
+					Entry->PoseWeight = WeightAndEasings ? *WeightAndEasings : 1.0;
+				}
+			}
+		}
+
+		TObjectPtr<UMovieSceneAnimMixerSystem> System;
+	};
+
+	struct FEvaluateAnimMixers
+	{
+		FEvaluateAnimMixers(UMovieSceneAnimMixerSystem* InSystem)
+			:System(InSystem)
+		{
+
+		}
+
+		void ForEachEntity(FMovieSceneEntityID EntityID,
+			FObjectComponent MeshComponent,
+			const TInstancedStruct<FMovieSceneMixedAnimationTarget>& Target,
+			TSharedPtr<FAnimNextEvaluationTask>& MixerTask) const
+		{
+			if (FMovieSceneAnimMixer* Mixer = System->Mixers.Find(FMovieSceneAnimMixerKey({ MeshComponent.GetObject(), Target })))
+			{
+				// Create a new eval program
+				Mixer->EvaluationProgram = MakeShared<UE::AnimNext::FEvaluationProgram>();
+
+				if (!MixerTask.IsValid())
+				{
+					MixerTask = MakeShared<FAnimNextExecuteProgramTask>();
+				}
+				TSharedPtr<FAnimNextExecuteProgramTask> EvalProgramTask = StaticCastSharedPtr<FAnimNextExecuteProgramTask>(MixerTask);
+				EvalProgramTask->Program = Mixer->EvaluationProgram;
+
+				// From lowest to highest priority, add the tasks to the program.
+				// Each task either
+				// a) Is self-contained, adding a single keyframe to the stack, in which case it will require a blend accumulation task afterwards
+				// b) Relies on the already accumulated pose and outputs a fully mixed pose to the stack, in which case no additional blend accumulation task is required
+				// Note that we assume here that the target will push some kind of 'base pose' to the VM before evaluating the task. This allows us to blend in/out from gameplay for example.
+
+				for (const FMovieSceneAnimMixerEntry& MixerEntry : Mixer->Entries)
+				{
+					if (MixerEntry.EvalTask.IsValid())
+					{
+						Mixer->EvaluationProgram->AppendTaskPtr(MixerEntry.EvalTask);
+
+						if (Mixer->RootMotion && MixerEntry.RootMotionSettings && MixerEntry.RootMotionSettings->RootMotionSpace == EMovieSceneRootMotionSpace::AnimationSpace)
+						{
+							// Need to convert root motion attribute to world space for blending, insert task to do that
+							Mixer->EvaluationProgram->AppendTask(FAnimNextConvertRootMotionToWorldSpaceTask::Make(Mixer->RootMotion));
+						}
+
+						if (MixerEntry.bRequiresBlend)
+						{
+							if (MixerEntry.bAdditive)
+							{
+								// Add additive blend task.
+								Mixer->EvaluationProgram->AppendTask(FAnimNextApplyAdditiveKeyframeTask::Make(MixerEntry.PoseWeight));
+							}
+							else
+							{
+								// Add absolute blend task
+								Mixer->EvaluationProgram->AppendTask(FAnimNextBlendTwoKeyframesTask::Make(MixerEntry.PoseWeight));
+							}
+						}
+					}
+				}
+
+				if (Mixer->RootMotion)
+				{
+					// Add task to store the final root motion result
+					Mixer->EvaluationProgram->AppendTask(FAnimNextStoreRootTransformTask::Make(Mixer->RootMotion));
+				}
+			}
+		}
+
+		TObjectPtr<UMovieSceneAnimMixerSystem> System;
+	};
+}
+
 UMovieSceneAnimMixerSystem::UMovieSceneAnimMixerSystem(const FObjectInitializer& ObjInit)
 	: Super(ObjInit)
 {
@@ -50,11 +153,10 @@ UMovieSceneAnimMixerSystem::UMovieSceneAnimMixerSystem(const FObjectInitializer&
 	{
 		FBuiltInComponentTypes* BuiltInComponents = FBuiltInComponentTypes::Get();
 
-		DefineImplicitPrerequisite(GetClass(), UMovieSceneEntityGroupingSystem::StaticClass());
-		DefineImplicitPrerequisite(GetClass(), UMovieScenePropertyInstantiatorSystem::StaticClass());
-		DefineImplicitPrerequisite(GetClass(), UMovieSceneInterrogatedPropertyInstantiatorSystem::StaticClass());
+		DefineImplicitPrerequisite(UMovieSceneGenericBoundObjectInstantiator::StaticClass(), GetClass());
 		DefineImplicitPrerequisite(UMovieSceneGenericBoundObjectInstantiator::StaticClass(), GetClass());
 		DefineImplicitPrerequisite(UMovieSceneBoundSceneComponentInstantiator::StaticClass(), GetClass());
+		DefineImplicitPrerequisite(UWeightAndEasingEvaluatorSystem::StaticClass(), GetClass());
 		DefineComponentConsumer(GetClass(), BuiltInComponents->BoundObjectKey);
 		DefineComponentConsumer(GetClass(), BuiltInComponents->BoundObject);
 		DefineComponentConsumer(GetClass(), AnimMixerComponents->Task);
@@ -102,7 +204,7 @@ void UMovieSceneAnimMixerSystem::OnRun(FSystemTaskPrerequisites& InPrerequisites
 		.Write(AnimMixerComponents->Target)
 		.Read(AnimMixerComponents->Task)
 		.Read(AnimMixerComponents->Priority)
-		.ReadOptional(AnimMixerComponents->RootMotionType)
+		.ReadOptional(AnimMixerComponents->RootMotionSettings)
 		.FilterAll({ BuiltInComponents->Tags.NeedsLink})
 		.Iterate_PerEntity( &Linker->EntityManager,
 		[this, &NewMixers, BuiltInComponents, AnimMixerComponents]
@@ -112,7 +214,7 @@ void UMovieSceneAnimMixerSystem::OnRun(FSystemTaskPrerequisites& InPrerequisites
 			TInstancedStruct<FMovieSceneMixedAnimationTarget>& Target,
 			TSharedPtr<FAnimNextEvaluationTask> Task,
 			int32 Priority,
-			const ESwapRootBone* RootMotionType)
+			const FMovieSceneRootMotionSettings* RootMotionSettings)
 			{
 				// For new entities, we 'resolve' the animation target so if 'Automatic' is picked we choose the right one automatically.
 				Target = ResolveAnimationTarget(BoundObjectKey, Target);
@@ -133,7 +235,10 @@ void UMovieSceneAnimMixerSystem::OnRun(FSystemTaskPrerequisites& InPrerequisites
 					NewEntry.EvalTask = Task;
 					NewEntry.Priority = Priority;
 					NewEntry.PoseWeight = 0;
-					NewEntry.RootMotionType = RootMotionType ? *RootMotionType : ESwapRootBone::SwapRootBone_None;
+					if (RootMotionSettings)
+					{
+						NewEntry.RootMotionSettings = *RootMotionSettings;
+					}
 					NewEntry.bAdditive = Type.ContainsAny({BuiltInComponents->Tags.AdditiveBlend, BuiltInComponents->Tags.AdditiveFromBaseBlend});
 					NewEntry.bRequiresBlend = Type.Contains(AnimMixerComponents->Tags.RequiresBlending);
 					int32 EntryIndex = Algo::LowerBound(Mixer->Entries, NewEntry);
@@ -150,178 +255,78 @@ void UMovieSceneAnimMixerSystem::OnRun(FSystemTaskPrerequisites& InPrerequisites
 				.Add(AnimMixerComponents->MeshComponent, FObjectComponent::Weak(NewMixerKey.BoundObjectKey.ResolveObjectPtr()))
 				.Add(AnimMixerComponents->Target, NewMixerKey.Target)
 				.Add(AnimMixerComponents->MixerTask, nullptr)
-				.Add(AnimMixerComponents->RootMotionType, ESwapRootBone::SwapRootBone_None)
 				.AddTag(BuiltInComponents->Tags.RestoreState) // TODO: For now we always restore state on the mixer when it gets unlinked.
 				.CreateEntity(&Linker->EntityManager);
 		}
 
 
 		// ----------------------------------------------------------------------------------------------------------------------------------------
-		// Step 4 - Update mixer entities, and remove stale ones
+		// Step 4 - Update mixer entities, and remove stale ones. 
+		// We also add/remove root motion components from the mixer entities depending on the entries present in them.
 		TArray<FMovieSceneAnimMixerKey> StaleMixers;
-
-		// Component mask for adding root motion support to mixers if necessary
-		FComponentMask RootMotionMixerComponents;
-		RootMotionMixerComponents.SetAll({
-			BuiltInComponents->PropertyBinding,
-			BuiltInComponents->BlenderType,
-			TrackComponents->ComponentTransform.PropertyTag,
-			BuiltInComponents->BoundObject, // The target of the root motion, to be filled in by the root motion system
-			BuiltInComponents->BoundObjectKey,
-			BuiltInComponents->Tags.AdditiveBlend, // We treat root motion as additive since it will be computed in component space.
-			BuiltInComponents->DoubleResult[0],
-			BuiltInComponents->DoubleResult[1],
-			BuiltInComponents->DoubleResult[2],
-			BuiltInComponents->DoubleResult[3],
-			BuiltInComponents->DoubleResult[4],
-			BuiltInComponents->DoubleResult[5],
-			BuiltInComponents->DoubleResult[6],
-			BuiltInComponents->DoubleResult[7],
-			BuiltInComponents->DoubleResult[8],
-			});
-
-		for (const TPair<FMovieSceneAnimMixerKey, FMovieSceneAnimMixer>& MixerPair : Mixers)
+		TArray<FObjectKey> StaleRootMotion;
+		TArray<TSharedPtr<FMovieSceneMixerRootMotionComponentData>> NewRootMotion;
+		for (TPair<FMovieSceneAnimMixerKey, FMovieSceneAnimMixer>& MixerPair : Mixers)
 		{
-			ESwapRootBone RootMotionType = ESwapRootBone::SwapRootBone_None;
-			const FMovieSceneAnimMixerEntry* MaxEntry = Algo::MaxElementBy(MixerPair.Value.Entries, [](const FMovieSceneAnimMixerEntry& Entry) { return Entry.RootMotionType; });
+			ESwapRootBone SwapRootBone = ESwapRootBone::SwapRootBone_None;
+			FMovieSceneAnimMixer& Mixer = MixerPair.Value;
+			const FMovieSceneAnimMixerEntry* MaxEntry = Algo::MaxElementBy(Mixer.Entries, [](const FMovieSceneAnimMixerEntry& Entry) { return Entry.RootMotionSettings ? Entry.RootMotionSettings->SwapRootBone : ESwapRootBone::SwapRootBone_None; });
 			if (MaxEntry)
 			{
-				RootMotionType = MaxEntry->RootMotionType;
+				SwapRootBone = MaxEntry->RootMotionSettings ? MaxEntry->RootMotionSettings->SwapRootBone : ESwapRootBone::SwapRootBone_None;
 			}
 
-			bool bHasRootMotionComponent = Linker->EntityManager.GetEntityType(MixerPair.Value.MixerEntityID).Contains(TrackComponents->ComponentTransform.PropertyTag);
+			bool bHasRootMotionComponent = Linker->EntityManager.GetEntityType(Mixer.MixerEntityID).Contains(AnimMixerComponents->MixerRootMotion);
 
-			if (MixerPair.Value.Entries.IsEmpty())
+			if (Mixer.Entries.IsEmpty())
 			{
 				StaleMixers.Add(MixerPair.Key);
-				Linker->EntityManager.AddComponent(MixerPair.Value.MixerEntityID, BuiltInComponents->Tags.NeedsUnlink, EEntityRecursion::Full);
+				Linker->EntityManager.AddComponent(Mixer.MixerEntityID, BuiltInComponents->Tags.NeedsUnlink, EEntityRecursion::Full);
 			}
-			else if (RootMotionType != ESwapRootBone::SwapRootBone_None) 
+			else if (SwapRootBone != ESwapRootBone::SwapRootBone_None)
 			{
-				USkeletalMeshComponent* SkeletalMeshComponent = Cast<USkeletalMeshComponent>(MixerPair.Key.BoundObjectKey.ResolveObjectPtr());
-				if (!SkeletalMeshComponent)
+				if (!Mixer.RootMotion.IsValid())
 				{
-					continue;
+					Mixer.RootMotion = MakeShared<FMovieSceneMixerRootMotionComponentData>();
 				}
-				UObject* RootMotionTarget = SkeletalMeshComponent;
-				if (RootMotionType == ESwapRootBone::SwapRootBone_Actor)
+				Mixer.RootMotion->MixerEntityID = Mixer.MixerEntityID;
+				
+				// Ensure actor matches
+				USkeletalMeshComponent* SkeletalMeshComponent = Cast<USkeletalMeshComponent>(MixerPair.Key.BoundObjectKey.ResolveObjectPtr());
+				USceneComponent* RootMotionTarget = SkeletalMeshComponent;
+				if (SwapRootBone == ESwapRootBone::SwapRootBone_Actor)
 				{
-					RootMotionTarget = SkeletalMeshComponent->GetOwner()->GetRootComponent();
+					RootMotionTarget = SkeletalMeshComponent ? SkeletalMeshComponent->GetOwner()->GetRootComponent() : nullptr;
+				}
+				Mixer.RootMotion->Target = RootMotionTarget;
+
+				if (Mixer.RootMotion->Target.IsValid() && !bHasRootMotionComponent)
+				{ 
+					// Add this root motion to the mixer entity
+					FEntityBuilder()
+						.Add(AnimMixerComponents->MixerRootMotion, Mixer.RootMotion)
+						.AddTag(BuiltInComponents->Tags.NeedsLink)
+						.MutateExisting(&Linker->EntityManager, Mixer.MixerEntityID);
 				}
-				// Update root motion type and add mixer components as necessary
-				FTypelessMutation Mutation;
-				Mutation.AddMask = RootMotionMixerComponents;
-				FEntityBuilder()
-					.Add(AnimMixerComponents->RootMotionType, RootMotionType)
-					.Add(BuiltInComponents->PropertyBinding, FMovieScenePropertyBinding(TEXT("Transform"), TEXT("Transform")))
-					.Add(BuiltInComponents->BlenderType, UMovieScenePiecewiseDoubleBlenderSystem::StaticClass())
-					.Add(BuiltInComponents->BoundObject, RootMotionTarget)
-					.Add(BuiltInComponents->BoundObjectKey, RootMotionTarget)
-					.AddTagConditional(BuiltInComponents->Tags.NeedsLink, !bHasRootMotionComponent)
-					.MutateExisting(&Linker->EntityManager, MixerPair.Value.MixerEntityID, Mutation);
 			}
-			else if (RootMotionType == ESwapRootBone::SwapRootBone_None && bHasRootMotionComponent)
+			else if (SwapRootBone == ESwapRootBone::SwapRootBone_None && bHasRootMotionComponent)
 			{
-				// Remove the root motion mixer components, but add NeedsLink again to poke the transform system to remove this
+				Mixer.RootMotion = nullptr;
+
+				// Remove root motion from the mixer entity
 				FTypelessMutation Mutation;
-				Mutation.RemoveMask = RootMotionMixerComponents;
+				Mutation.Remove({ AnimMixerComponents->MixerRootMotion });
 				FEntityBuilder()
 					.AddTag(BuiltInComponents->Tags.NeedsLink)
-					.MutateExisting(&Linker->EntityManager, MixerPair.Value.MixerEntityID, Mutation);
+					.MutateExisting(&Linker->EntityManager, Mixer.MixerEntityID, Mutation);
 			}
 		}
+
+		// Remove stale mixers
 		Algo::ForEach(StaleMixers, [this](const FMovieSceneAnimMixerKey& MixerKey) { Mixers.Remove(MixerKey); });
 	}
 }
 
-namespace UE::MovieScene
-{
-	struct FUpdateTaskEntities
-	{
-		FUpdateTaskEntities(UMovieSceneAnimMixerSystem* InSystem)
-			:System(InSystem)
-
-		{
-
-		}
-		void ForEachEntity(FMovieSceneEntityID EntityID, 
-							FObjectKey BoundObjectKey,
-							const TInstancedStruct<FMovieSceneMixedAnimationTarget>& Target,
-							TSharedPtr<FAnimNextEvaluationTask> Task,
-							const double* WeightAndEasings) const
-		{
-			FMovieSceneAnimMixerKey Key({ BoundObjectKey, Target });
-			if (FMovieSceneAnimMixer* Mixer = System->Mixers.Find(Key))
-			{
-				// Find the entry in question
-				if (FMovieSceneAnimMixerEntry* Entry = Mixer->Entries.FindByPredicate([EntityID](const FMovieSceneAnimMixerEntry& TestEntry) { return TestEntry.EntityID == EntityID; }))
-				{
-					Entry->EvalTask = Task;
-					Entry->PoseWeight = WeightAndEasings ? *WeightAndEasings : 1.0;
-				}
-			}
-		}
-
-		TObjectPtr<UMovieSceneAnimMixerSystem> System;
-	};
-
-	struct FEvaluateAnimMixers
-	{
-		FEvaluateAnimMixers(UMovieSceneAnimMixerSystem* InSystem)
-			:System(InSystem)
-		{
-
-		}
-
-		void ForEachEntity(FMovieSceneEntityID EntityID,
-			FObjectComponent MeshComponent,
-			const TInstancedStruct<FMovieSceneMixedAnimationTarget>& Target,
-			TSharedPtr<FAnimNextEvaluationTask>& MixerTask) const
-		{
-			if (FMovieSceneAnimMixer* Mixer = System->Mixers.Find(FMovieSceneAnimMixerKey({ MeshComponent.GetObject(), Target })))
-			{
-				// Create a new eval program
-				Mixer->EvaluationProgram = MakeShared<UE::AnimNext::FEvaluationProgram>();
-
-				if (!MixerTask.IsValid())
-				{
-					MixerTask = MakeShared<FAnimNextExecuteProgramTask>();
-				}
-				TSharedPtr<FAnimNextExecuteProgramTask> EvalProgramTask = StaticCastSharedPtr<FAnimNextExecuteProgramTask>(MixerTask);
-				EvalProgramTask->Program = Mixer->EvaluationProgram;
-
-				// From lowest to highest priority, add the tasks to the program.
-				// Each task either
-				// a) Is self-contained, adding a single keyframe to the stack, in which case it will require a blend accumulation task afterwards
-				// b) Relies on the already accumulated pose and outputs a fully mixed pose to the stack, in which case no additional blend accumulation task is required
-				// Note that we assume here that the target will push some kind of 'base pose' to the VM before evaluating the task. This allows us to blend in/out from gameplay for example.
-				for (const FMovieSceneAnimMixerEntry& MixerEntry : Mixer->Entries)
-				{
-					if (MixerEntry.EvalTask.IsValid())
-					{
-						Mixer->EvaluationProgram->AppendTaskPtr(MixerEntry.EvalTask);
-						if (MixerEntry.bRequiresBlend)
-						{
-							if (MixerEntry.bAdditive)
-							{
-								// Add additive blend task.
-								Mixer->EvaluationProgram->AppendTask(FAnimNextApplyAdditiveKeyframeTask::Make(MixerEntry.PoseWeight));
-							}
-							else
-							{
-								// Add absolute blend task
-								Mixer->EvaluationProgram->AppendTask(FAnimNextBlendTwoKeyframesTask::Make(MixerEntry.PoseWeight));
-							}
-						}
-					}
-				}
-			}
-		}
-
-		TObjectPtr<UMovieSceneAnimMixerSystem> System;
-	};
-}
-
 void UMovieSceneAnimMixerSystem::OnSchedulePersistentTasks(UE::MovieScene::IEntitySystemScheduler* TaskScheduler)
 {
 	using namespace UE::MovieScene;
@@ -394,7 +399,7 @@ void UMovieSceneAnimMixerSystem::OnCleanTaggedGarbage()
 const TInstancedStruct<FMovieSceneMixedAnimationTarget> UMovieSceneAnimMixerSystem::ResolveAnimationTarget(FObjectKey ObjectKey, const TInstancedStruct<FMovieSceneMixedAnimationTarget> InTarget)
 {
 	// If user has selected the default 'automatic' target, attempt to choose one automatically for them.
-	if (InTarget.GetScriptStruct() == FMovieSceneMixedAnimationTarget::StaticStruct())
+	if (!InTarget.IsValid() || InTarget.GetScriptStruct() == FMovieSceneMixedAnimationTarget::StaticStruct())
 	{
 		if (UObject* Object = ObjectKey.ResolveObjectPtr())
 		{
diff --git a/Engine/Plugins/Experimental/MovieSceneAnimMixer/Source/MovieSceneAnimMixer/Private/Systems/MovieSceneMixedSkeletalAnimationSystem.cpp b/Engine/Plugins/Experimental/MovieSceneAnimMixer/Source/MovieSceneAnimMixer/Private/Systems/MovieSceneMixedSkeletalAnimationSystem.cpp
--- a/Engine/Plugins/Experimental/MovieSceneAnimMixer/Source/MovieSceneAnimMixer/Private/Systems/MovieSceneMixedSkeletalAnimationSystem.cpp
+++ b/Engine/Plugins/Experimental/MovieSceneAnimMixer/Source/MovieSceneAnimMixer/Private/Systems/MovieSceneMixedSkeletalAnimationSystem.cpp
@@ -65,7 +65,7 @@ struct FSkeletalAnimMixerMutation : IMovieSceneEntityMutation
 			AnimMixerComponents->Priority,
 			AnimMixerComponents->Target,
 			AnimMixerComponents->Task,
-			AnimMixerComponents->RootMotionType,
+			AnimMixerComponents->RootMotionSettings,
 			AnimMixerComponents->Tags.RequiresBlending
 		});
 	}
@@ -78,7 +78,7 @@ struct FSkeletalAnimMixerMutation : IMovieSceneEntityMutation
 		TComponentWriter<TInstancedStruct<FMovieSceneMixedAnimationTarget>> OutAnimTargets = Allocation->WriteComponents(AnimMixerComponents->Target, WriteContext);
 		TComponentWriter<TSharedPtr<FAnimNextEvaluationTask>> OutTasks = Allocation->WriteComponents(AnimMixerComponents->Task, WriteContext);
 		TComponentWriter<int32> OutPriorities = Allocation->WriteComponents(AnimMixerComponents->Priority, WriteContext);
-		TComponentWriter<ESwapRootBone> OutRootMotionTypes = Allocation->WriteComponents(AnimMixerComponents->RootMotionType, WriteContext);
+		TComponentWriter<FMovieSceneRootMotionSettings> OutRootMotionSettings = Allocation->WriteComponents(AnimMixerComponents->RootMotionSettings, WriteContext);
 		for (int32 Index = 0; Index < Num; ++Index)
 		{
 			// Initialize a null (empty) task, will be written during evaluation
@@ -99,7 +99,8 @@ struct FSkeletalAnimMixerMutation : IMovieSceneEntityMutation
 				}
 
 				OutPriorities[Index] = SkeletalAnim.Section->MixedAnimationPriority;
-				OutRootMotionTypes[Index] = SkeletalAnim.Section->Params.SwapRootBone;
+				OutRootMotionSettings[Index].SwapRootBone = SkeletalAnim.Section->Params.SwapRootBone;
+				OutRootMotionSettings[Index].RootMotionSpace = EMovieSceneRootMotionSpace::AnimationSpace;
 			}
 		}
 	}
@@ -248,9 +249,6 @@ void FMovieSceneSkeletalAnimationEvaluationTask::Execute(UE::AnimNext::FEvaluati
 
 		const bool bUseRawData = FDecompressionTools::ShouldUseRawData(AnimSequencePtr, Keyframe.Pose);
 
-
-		bool bGrabbedRootTransform = false;
-
 		TOptional<FTransform> RootBoneTransform;
 
 		if (EnumHasAnyFlags(VM.GetFlags(), EEvaluationFlags::Bones))
diff --git a/Engine/Plugins/Experimental/MovieSceneAnimMixer/Source/MovieSceneAnimMixer/Private/Systems/MovieSceneRootMotionSystem.cpp b/Engine/Plugins/Experimental/MovieSceneAnimMixer/Source/MovieSceneAnimMixer/Private/Systems/MovieSceneRootMotionSystem.cpp
--- a/Engine/Plugins/Experimental/MovieSceneAnimMixer/Source/MovieSceneAnimMixer/Private/Systems/MovieSceneRootMotionSystem.cpp
+++ b/Engine/Plugins/Experimental/MovieSceneAnimMixer/Source/MovieSceneAnimMixer/Private/Systems/MovieSceneRootMotionSystem.cpp
@@ -22,117 +22,256 @@
 #include "EntitySystem/MovieSceneEntitySystemLinker.h"
 #include "Systems/MovieScenePropertyInstantiator.h"
 #include "EntitySystem/Interrogation/MovieSceneInterrogatedPropertyInstantiator.h"
+#include "EntitySystem/MovieSceneEntitySystemRunner.h"
+#include "Misc/ScopeRWLock.h"
+
+#if WITH_EDITOR
+#include "Editor.h"
+#include "Editor/UnrealEdEngine.h"
+#include "UnrealEdGlobals.h"
+#endif
 
 #include UE_INLINE_GENERATED_CPP_BY_NAME(MovieSceneRootMotionSystem)
 
 namespace UE::MovieScene
 {
-	struct FEvaluateRootMotion
+	struct FUpdateRootMotionData
 	{
-		void ForEachEntity
-		(
-			FObjectComponent MeshComponent,
-			TSharedPtr<FAnimNextEvaluationTask> MixerTask,
-			ESwapRootBone RootMotionType,
-			double& DoubleResult0,
-			double& DoubleResult1,
-			double& DoubleResult2,
-			double& DoubleResult3,
-			double& DoubleResult4,
-			double& DoubleResult5,
-			double& DoubleResult6,
-			double& DoubleResult7,
-			double& DoubleResult8
-		) const
+		void ForEachAllocation(
+			const FEntityAllocation* Allocation,
+			TRead<FObjectComponent> MeshComponents,
+			TWrite<TSharedPtr<FMovieSceneMixerRootMotionComponentData>> RootMotions) const
 		{
 			using namespace UE::AnimNext;
 			using namespace UE::MovieScene;
 
-			// Zero out the root motion to start
-			DoubleResult0 = 0;
-			DoubleResult1 = 0;
-			DoubleResult2 = 0;
-			DoubleResult3 = 0;
-			DoubleResult4 = 0;
-			DoubleResult5 = 0;
-			DoubleResult6 = 0;
-			DoubleResult7 = 0;
-			DoubleResult8 = 0;
-
-			if (RootMotionType == ESwapRootBone::SwapRootBone_None)
+			for (int32 Index = 0; Index < Allocation->Num(); ++Index)
 			{
-				
-				return;
+				if (TSharedPtr<FMovieSceneMixerRootMotionComponentData> RootMotion = RootMotions[Index])
+				{
+					USkeletalMeshComponent* SkeletalMeshComponent = Cast<USkeletalMeshComponent>(MeshComponents[Index].GetObject());
+					if (SkeletalMeshComponent && SkeletalMeshComponent != SkeletalMeshComponent->GetOwner()->GetRootComponent() && SkeletalMeshComponent != RootMotion->Target)
+					{
+						RootMotion->SetInverseMeshToActorRotation(SkeletalMeshComponent->GetOwner()->GetRootComponent()->GetComponentTransform().GetRelativeTransformReverse(SkeletalMeshComponent->GetComponentTransform()).GetRotation());
+					}
+				}
 			}
+		}
+	};
 
-			FAnimMixerComponentTypes* AnimMixerTypes = FAnimMixerComponentTypes::Get();
+	struct FSyncTransformTrackResult
+	{
+		UMovieSceneRootMotionSystem* System = nullptr;
+
+		FMovieSceneTracksComponentTypes* TrackComponents = nullptr;
+		
+		FSyncTransformTrackResult(UMovieSceneRootMotionSystem* InRootMotionSystem)
+		: System(InRootMotionSystem)
+		{
+			TrackComponents = FMovieSceneTracksComponentTypes::Get();
+		}
 
-			USkeletalMeshComponent* SkeletalMeshComponent = Cast<USkeletalMeshComponent>(MeshComponent.GetObject());
-			ensureMsgf(SkeletalMeshComponent, TEXT("Attempting to apply animation to an anim instance without a valid skeletal mesh component."));
+		void ForEachAllocation(
+			const FEntityAllocation* Allocation,
+			TRead<FObjectKey> RootMotionTargets,
+			TReadOneOrMoreOf<
+			double, double, double,
+			double, double, double,
+			double, double, double>
+			TransformBlendResults) const
+		{
+			using namespace UE::AnimNext;
+			using namespace UE::MovieScene;
 
-			if (!SkeletalMeshComponent || !SkeletalMeshComponent->GetSkeletalMeshAsset())
+			const double* DoubleResults[9];
 			{
-				return;
+				DoubleResults[0] = TransformBlendResults.Get<0>();
+				DoubleResults[1] = TransformBlendResults.Get<1>();
+				DoubleResults[2] = TransformBlendResults.Get<2>();
+				DoubleResults[3] = TransformBlendResults.Get<3>();
+				DoubleResults[4] = TransformBlendResults.Get<4>();
+				DoubleResults[5] = TransformBlendResults.Get<5>();
+				DoubleResults[6] = TransformBlendResults.Get<6>();
+				DoubleResults[7] = TransformBlendResults.Get<7>();
+				DoubleResults[8] = TransformBlendResults.Get<8>();
 			}
 
-			if (MixerTask.IsValid())
-			{
-				FMemMark Mark(FMemStack::Get());
-				FDataHandle RefPoseHandle = FDataRegistry::Get()->GetOrGenerateReferencePose(SkeletalMeshComponent);
-				FAnimNextGraphReferencePose GraphReferencePose(RefPoseHandle);
+			const bool bFullyAnimated =
+				DoubleResults[0] && DoubleResults[1] && DoubleResults[2] &&
+				DoubleResults[3] && DoubleResults[4] && DoubleResults[5] &&
+				DoubleResults[6] && DoubleResults[7] && DoubleResults[8];
 
-				const FReferencePose& RefPose = RefPoseHandle.GetRef<FReferencePose>();
-				TOptional<FTransform> ResultTransform;
-				{
-					// Mix just the root motion
-					FEvaluationVM EvaluationVM(EEvaluationFlags::Trajectory, RefPose, 0);
 
-					MixerTask->Execute(EvaluationVM);
+			TOptionalComponentReader<FIntermediate3DTransform> OptInitialValues =
+				Allocation->TryReadComponents(TrackComponents->Parameters.Transform.InitialValue);
 
-					TUniquePtr<FKeyframeState> EvaluatedKeyframe;
-					if (EvaluationVM.PopValue(KEYFRAME_STACK_NAME, EvaluatedKeyframe))
+			for (int32 Index = 0; Index < Allocation->Num(); ++Index)
+			{
+				if (TSharedPtr<FMovieSceneMixerRootMotionComponentData>* RootMotionDataPtr = System->GetData().Find(RootMotionTargets[Index]))
+				{
+					// We have root motion for this object. Copy over the transform.
+					if (RootMotionDataPtr->IsValid())
 					{
-						if (EvaluatedKeyframe.IsValid())
+						if (bFullyAnimated)
+						{
+							(*RootMotionDataPtr)->SetSequenceAnimationTransform(FTransform(
+								FRotator(DoubleResults[4][Index], DoubleResults[5][Index], DoubleResults[3][Index]),
+								FVector(DoubleResults[0][Index], DoubleResults[1][Index], DoubleResults[2][Index]),
+								FVector(DoubleResults[6][Index], DoubleResults[7][Index], DoubleResults[8][Index])
+							));
+						}
+						else
 						{
-							if (FTransformAnimationAttribute* RootMotionAttribute = EvaluatedKeyframe->Attributes.Find<FTransformAnimationAttribute>(AnimMixerTypes->RootTransformAttributeId))
+							FRotator Rotation = FRotator::ZeroRotator;
+							FVector Location = FVector::ZeroVector;
+							FVector Scale = FVector::OneVector;
+
+							if (OptInitialValues)
 							{
-								ResultTransform = RootMotionAttribute->Value;
+								Rotation = OptInitialValues[Index].GetRotation();
+								Location = OptInitialValues[Index].GetTranslation();
+								Scale = OptInitialValues[Index].GetScale();
 							}
+
+							if (DoubleResults[0]) { Location.X = DoubleResults[0][Index]; }
+							if (DoubleResults[1]) { Location.Y = DoubleResults[1][Index]; }
+							if (DoubleResults[2]) { Location.Z = DoubleResults[2][Index]; }
+
+							if (DoubleResults[3]) { Rotation.Roll = DoubleResults[3][Index]; }
+							if (DoubleResults[4]) { Rotation.Pitch = DoubleResults[4][Index]; }
+							if (DoubleResults[5]) { Rotation.Yaw = DoubleResults[5][Index]; }
+
+							if (DoubleResults[6]) { Scale.X = DoubleResults[6][Index]; }
+							if (DoubleResults[7]) { Scale.Y = DoubleResults[7][Index]; }
+							if (DoubleResults[8]) { Scale.Z = DoubleResults[8][Index]; }
+
+							(*RootMotionDataPtr)->SetSequenceAnimationTransform(FTransform(Rotation, Location, Scale));
 						}
 					}
 				}
+			}
+		}
+	};
+};
+
+
+TOptional<FTransform> FMovieSceneMixerRootMotionComponentData::GetSequenceAnimationTransform() const
+{
+	UE::TReadScopeLock ScopeLock(RootMotionLock);
+	return SequenceAnimationTransform;
+}
+
+void FMovieSceneMixerRootMotionComponentData::SetSequenceAnimationTransform(const FTransform& InTransform)
+{
+	UE::TWriteScopeLock ScopeLock(RootMotionLock);
+	SequenceAnimationTransform = InTransform;
+}
+
+TOptional<FQuat> FMovieSceneMixerRootMotionComponentData::GetInverseMeshToActorRotation() const
+{
+	UE::TReadScopeLock ScopeLock(RootMotionLock);
+	return InverseMeshToActorRotation;
+}
 
-				if (ResultTransform.IsSet())
+void FMovieSceneMixerRootMotionComponentData::SetInverseMeshToActorRotation(const FQuat& InInverseMeshToActorRotation)
+{
+	UE::TWriteScopeLock ScopeLock(RootMotionLock);
+	InverseMeshToActorRotation = InInverseMeshToActorRotation;
+}
+
+FTransform FMovieSceneMixerRootMotionComponentData::GetRootTransform() const
+{
+	UE::TReadScopeLock ScopeLock(RootMotionLock);
+	return RootTransform;
+}
+
+void FMovieSceneMixerRootMotionComponentData::SetRootTransform(const FTransform& InTransform)
+{
+	UE::TWriteScopeLock ScopeLock(RootMotionLock);
+	RootTransform = InTransform;
+}
+
+
+void FMovieSceneMixerRootMotionComponentData::Reset()
+{
+	UE::TWriteScopeLock ScopeLock(RootMotionLock);
+	SequenceAnimationTransform.Reset();
+	InverseMeshToActorRotation.Reset();
+	RootTransform.SetIdentity();
+}
+
+FAnimNextConvertRootMotionToWorldSpaceTask::FAnimNextConvertRootMotionToWorldSpaceTask(const TSharedPtr<FMovieSceneMixerRootMotionComponentData>& InRootMotionData)
+	: RootMotionData(InRootMotionData)
+{
+}
+
+FAnimNextConvertRootMotionToWorldSpaceTask FAnimNextConvertRootMotionToWorldSpaceTask::Make(const TSharedPtr<FMovieSceneMixerRootMotionComponentData>& InRootMotionData)
+{
+	return FAnimNextConvertRootMotionToWorldSpaceTask(InRootMotionData);
+}
+
+void FAnimNextConvertRootMotionToWorldSpaceTask::Execute(UE::AnimNext::FEvaluationVM& VM) const
+{
+	using namespace UE::AnimNext; 
+	using namespace UE::MovieScene;
+
+	FAnimMixerComponentTypes* AnimMixerComponents = FAnimMixerComponentTypes::Get();
+	TOptional<FTransform> SequenceAnimationTransform = RootMotionData.IsValid() ? RootMotionData->GetSequenceAnimationTransform() : TOptional<FTransform>();
+	if (SequenceAnimationTransform.IsSet() && EnumHasAnyFlags(VM.GetFlags(), EEvaluationFlags::Attributes | EEvaluationFlags::Trajectory))
+	{
+		if (TUniquePtr<FKeyframeState>* Keyframe = VM.PeekValueMutable<TUniquePtr<FKeyframeState>>(KEYFRAME_STACK_NAME, 0))
+		{
+			if (Keyframe->IsValid())
+			{
+				if (FTransformAnimationAttribute* RootMotionAttribute = (*Keyframe)->Attributes.Find<FTransformAnimationAttribute>(AnimMixerComponents->RootTransformAttributeId))
 				{
-					if (RootMotionType == ESwapRootBone::SwapRootBone_Actor)
+					FTransform RelativeTransform = RootMotionAttribute->Value;
+					TOptional<FQuat> InverseMeshToActorRotation = RootMotionData->GetInverseMeshToActorRotation();
+					if (InverseMeshToActorRotation.IsSet())
 					{
-						AActor* Actor = SkeletalMeshComponent->GetOwner();
-						if (Actor && Actor->GetRootComponent())
-						{
-							// Compensate for any mesh component rotation
-							FQuat InverseMeshToActorRotation = SkeletalMeshComponent->GetOwner()->GetRootComponent()->GetComponentTransform().GetRelativeTransformReverse(SkeletalMeshComponent->GetComponentTransform()).GetRotation();
-							ResultTransform->SetTranslation(InverseMeshToActorRotation * ResultTransform->GetTranslation());
-						}
+						RelativeTransform.SetTranslation(InverseMeshToActorRotation.GetValue() * RelativeTransform.GetTranslation());
 					}
+					RelativeTransform = RelativeTransform * SequenceAnimationTransform.GetValue();
+					RootMotionAttribute->Value = RelativeTransform;
+				}
+			}
+		}
+	}
+}
+
+FAnimNextStoreRootTransformTask::FAnimNextStoreRootTransformTask(const TSharedPtr<FMovieSceneMixerRootMotionComponentData>& InRootMotionData)
+	: RootMotionData(InRootMotionData)
+{
 
-					DoubleResult0 = ResultTransform->GetTranslation().X;
-					DoubleResult1 = ResultTransform->GetTranslation().Y;
-					DoubleResult2 = ResultTransform->GetTranslation().Z;
+}
+
+FAnimNextStoreRootTransformTask FAnimNextStoreRootTransformTask::Make(const TSharedPtr<FMovieSceneMixerRootMotionComponentData>& InRootMotionData)
+{
+	return FAnimNextStoreRootTransformTask(InRootMotionData);
+}
 
-					FVector Euler = ResultTransform->Rotator().Euler();
-					DoubleResult3 = Euler.X;
-					DoubleResult4 = Euler.Y;
-					DoubleResult5 = Euler.Z;
 
-					// For now ignore scale. There's probably a better way to do this.
-					DoubleResult6 = 0; //ResultTransform->GetScale3D().X;
-					DoubleResult7 = 0; //ResultTransform->GetScale3D().Y;
-					DoubleResult8 = 0; //ResultTransform->GetScale3D().Z;*/
+void FAnimNextStoreRootTransformTask::Execute(UE::AnimNext::FEvaluationVM& VM) const
+{
+	using namespace UE::AnimNext;
+	using namespace UE::MovieScene;
+
+	FAnimMixerComponentTypes* AnimMixerComponents = FAnimMixerComponentTypes::Get();
+
+	if (RootMotionData.IsValid() && EnumHasAnyFlags(VM.GetFlags(), EEvaluationFlags::Attributes | EEvaluationFlags::Trajectory))
+	{
+		if (TUniquePtr<FKeyframeState>* Keyframe = VM.PeekValueMutable<TUniquePtr<FKeyframeState>>(KEYFRAME_STACK_NAME, 0))
+		{
+			if (Keyframe->IsValid())
+			{
+				if (FTransformAnimationAttribute* RootMotionAttribute = (*Keyframe)->Attributes.Find<FTransformAnimationAttribute>(AnimMixerComponents->RootTransformAttributeId))
+				{
+					RootMotionData->SetRootTransform(RootMotionAttribute->Value);
 				}
 			}
 		}
-	};
-};
+	}
+}
 
 UMovieSceneRootMotionSystem::UMovieSceneRootMotionSystem(const FObjectInitializer& ObjInit)
 	: Super(ObjInit)
@@ -141,39 +280,176 @@ UMovieSceneRootMotionSystem::UMovieSceneRootMotionSystem(const FObjectInitialize
 
 	FAnimMixerComponentTypes* AnimMixerComponents = FAnimMixerComponentTypes::Get();
 
-	RelevantComponent = AnimMixerComponents->MixerTask;
-	Phase = ESystemPhase::Scheduling;
+	RelevantComponent = AnimMixerComponents->MixerRootMotion;
+	Phase = ESystemPhase::Instantiation | ESystemPhase::Scheduling;
 
 	SystemCategories |= FSystemInterrogator::GetExcludedFromInterrogationCategory();
 
 	if (HasAnyFlags(RF_ClassDefaultObject))
 	{
-		DefineImplicitPrerequisite(GetClass(), UMovieScenePiecewiseDoubleBlenderSystem::StaticClass());
+		DefineImplicitPrerequisite(UMovieScenePiecewiseDoubleBlenderSystem::StaticClass(), GetClass());
 		DefineImplicitPrerequisite(UMovieSceneAnimMixerSystem::StaticClass(), GetClass()); 
 	}
 }
 
+#if WITH_EDITOR
+void UMovieSceneRootMotionSystem::OnLink()
+{
+	if (OnBeginActorMovementHandle.IsValid() == false)
+	{
+		OnBeginActorMovementHandle = GEditor->OnBeginObjectMovement().Add(UEditorEngine::FOnBeginTransformObject::FDelegate::CreateLambda([this](UObject&)
+			{
+				bActorBeingMoved = true;
+			}));
+	}
+
+	if (OnEndActorMovementHandle.IsValid() == false)
+	{
+		OnEndActorMovementHandle = GEditor->OnEndObjectMovement().Add(UEditorEngine::FOnEndTransformObject::FDelegate::CreateLambda([this](UObject&)
+			{
+				bActorBeingMoved = false;
+			}));
+	}
+}
+
+void UMovieSceneRootMotionSystem::OnUnlink()
+{
+	GEditor->OnBeginObjectMovement().Remove(OnBeginActorMovementHandle);
+	OnBeginActorMovementHandle.Reset();
+	GEditor->OnEndObjectMovement().Remove(OnEndActorMovementHandle);
+	OnEndActorMovementHandle.Reset();
+}
+
+#endif
+
+void UMovieSceneRootMotionSystem::FinishDestroy()
+{
+	Super::FinishDestroy();
+}
+
+
+void UMovieSceneRootMotionSystem::OnRun(FSystemTaskPrerequisites& InPrerequisites, FSystemSubsequentTasks& Subsequents)
+{
+	using namespace UE::MovieScene;
+
+	FBuiltInComponentTypes* BuiltInComponents = FBuiltInComponentTypes::Get();
+	FMovieSceneTracksComponentTypes* TrackComponents = FMovieSceneTracksComponentTypes::Get();
+	FAnimMixerComponentTypes* AnimMixerComponents = FAnimMixerComponentTypes::Get();
+	TSharedRef<FMovieSceneEntitySystemRunner> Runner = Linker->GetRunner();
+
+	if (Runner->GetCurrentPhase() == ESystemPhase::Instantiation)
+	{
+		TArray<FMovieSceneEntityID> StaleRootMotion;
+		// Remove stale root motion from expiring mixers
+		FEntityTaskBuilder()
+			.ReadEntityIDs()
+			.FilterAll({ AnimMixerComponents->MixerTask, BuiltInComponents->Tags.NeedsUnlink })
+			.Iterate_PerEntity(&Linker->EntityManager, [&StaleRootMotion](FMovieSceneEntityID EntityID)
+				{
+					StaleRootMotion.Add(EntityID);
+				});
+
+		// Remove stale root motion from mixers without root motion components that need linking
+		FEntityTaskBuilder()
+			.ReadEntityIDs()
+			.FilterNone({ AnimMixerComponents->MixerRootMotion })
+			.FilterAll({ BuiltInComponents->Tags.NeedsLink, AnimMixerComponents->MixerTask})
+			.Iterate_PerEntity(&Linker->EntityManager, [&StaleRootMotion](FMovieSceneEntityID EntityID)
+				{
+					StaleRootMotion.Add(EntityID);
+				});
+
+		if (!StaleRootMotion.IsEmpty())
+		{
+			for (auto It = RootMotionData.CreateIterator(); It; ++It)
+			{
+				if (It->Value.IsValid() && StaleRootMotion.Contains(It->Value->MixerEntityID))
+				{
+					It.RemoveCurrent();
+				}
+			}
+		}
+
+		// Track or update root motion
+		FEntityTaskBuilder()
+		.ReadEntityIDs()
+		.Read(AnimMixerComponents->MixerRootMotion)
+		.FilterAll({BuiltInComponents->Tags.NeedsLink })
+		.Iterate_PerEntity(&Linker->EntityManager, [this](FMovieSceneEntityID EntityID, TSharedPtr<FMovieSceneMixerRootMotionComponentData> InRootMotion)
+			{
+				// Clear any existing data on that entity
+				for (auto It = RootMotionData.CreateIterator(); It; ++It)
+				{
+					if (It->Value.IsValid() && It->Value->MixerEntityID == EntityID)
+					{
+						It.RemoveCurrent();
+					}
+				}
+
+				if (InRootMotion.IsValid() && InRootMotion->Target.IsValid())
+				{
+					RootMotionData.Add(InRootMotion->Target.Get(), InRootMotion);
+				}
+			});
+	}
+}
+
 void UMovieSceneRootMotionSystem::OnSchedulePersistentTasks(UE::MovieScene::IEntitySystemScheduler* TaskScheduler)
 {
 	using namespace UE::MovieScene;
 
 	FBuiltInComponentTypes* BuiltInComponents = FBuiltInComponentTypes::Get();
+	FMovieSceneTracksComponentTypes* TrackComponents = FMovieSceneTracksComponentTypes::Get();
 	FAnimMixerComponentTypes* AnimMixerComponents = FAnimMixerComponentTypes::Get();
 
-	// Evaluate root motion
-	FTaskID EvaluateTask = FEntityTaskBuilder()
+	// Zero out all transforms to start
+	for (TTuple<FObjectKey, TSharedPtr<FMovieSceneMixerRootMotionComponentData>>& RootMotion : RootMotionData)
+	{
+		RootMotion.Value->Reset();
+	}
+
+	// Update some optional root motion data used later in the frame for calculation
+	FEntityTaskBuilder()
 		.Read(AnimMixerComponents->MeshComponent)
-		.Read(AnimMixerComponents->MixerTask)
-		.Read(AnimMixerComponents->RootMotionType)
-		.Write(BuiltInComponents->DoubleResult[0])
-		.Write(BuiltInComponents->DoubleResult[1])
-		.Write(BuiltInComponents->DoubleResult[2])
-		.Write(BuiltInComponents->DoubleResult[3])
-		.Write(BuiltInComponents->DoubleResult[4])
-		.Write(BuiltInComponents->DoubleResult[5])
-		.Write(BuiltInComponents->DoubleResult[6])
-		.Write(BuiltInComponents->DoubleResult[7])
-		.Write(BuiltInComponents->DoubleResult[8])
+		.Write(AnimMixerComponents->MixerRootMotion)
 		.SetDesiredThread(Linker->EntityManager.GetGatherThread())
-		.Schedule_PerEntity<FEvaluateRootMotion>(&Linker->EntityManager, TaskScheduler);
-}
\ No newline at end of file
+		.Schedule_PerAllocation<FUpdateRootMotionData>(&Linker->EntityManager, TaskScheduler);
+
+	// Copy the component transform blend result to the root motion component data
+	FEntityTaskBuilder()
+	.Read(BuiltInComponents->BoundObjectKey)
+	.ReadOneOrMoreOf(BuiltInComponents->DoubleResult[0],
+					BuiltInComponents->DoubleResult[1],
+					BuiltInComponents->DoubleResult[2],
+					BuiltInComponents->DoubleResult[3],
+					BuiltInComponents->DoubleResult[4],
+					BuiltInComponents->DoubleResult[5],
+					BuiltInComponents->DoubleResult[6],
+					BuiltInComponents->DoubleResult[7],
+					BuiltInComponents->DoubleResult[8])
+	.FilterAll({TrackComponents->ComponentTransform.PropertyTag})
+	.SetDesiredThread(Linker->EntityManager.GetGatherThread())
+	.Fork_PerAllocation<FSyncTransformTrackResult>(&Linker->EntityManager, TaskScheduler, this);
+}
+
+void UMovieSceneRootMotionSystem::ApplyRootMotion()
+{
+#if WITH_EDITOR
+	if (bActorBeingMoved)
+	{
+		return;
+	}
+#endif
+
+	for (TTuple<FObjectKey, TSharedPtr<FMovieSceneMixerRootMotionComponentData>> RootMotionTuple : RootMotionData)
+	{
+		if (TSharedPtr<FMovieSceneMixerRootMotionComponentData> RootMotion = RootMotionTuple.Value)
+		{
+			if (RootMotion->Target.IsValid())
+			{
+				FTransform RelativeTransform = RootMotion->GetRootTransform();
+				RootMotion->Target->SetRelativeLocationAndRotation(RelativeTransform.GetLocation(), RelativeTransform.GetRotation().Rotator());
+			}
+		}
+	}
+}
diff --git a/Engine/Plugins/Experimental/MovieSceneAnimMixer/Source/MovieSceneAnimMixer/Public/AnimMixerComponentTypes.h b/Engine/Plugins/Experimental/MovieSceneAnimMixer/Source/MovieSceneAnimMixer/Public/AnimMixerComponentTypes.h
--- a/Engine/Plugins/Experimental/MovieSceneAnimMixer/Source/MovieSceneAnimMixer/Public/AnimMixerComponentTypes.h
+++ b/Engine/Plugins/Experimental/MovieSceneAnimMixer/Source/MovieSceneAnimMixer/Public/AnimMixerComponentTypes.h
@@ -11,6 +11,7 @@
 #include "StructUtils/InstancedStruct.h"
 #include "EvaluationVM/EvaluationTask.h"
 #include "AnimSequencerInstanceProxy.h"
+#include "Systems/MovieSceneRootMotionSystem.h"
 
 namespace UE::MovieScene
 {
@@ -20,14 +21,15 @@ namespace UE::MovieScene
 		MOVIESCENEANIMMIXER_API static void Destroy();
 
 		// A root motion transform in animation space to be applied (after space conversion)
-		static const UE::Anim::FAttributeId RootTransformAttributeId;
+		MOVIESCENEANIMMIXER_API static const UE::Anim::FAttributeId RootTransformAttributeId;
 
 		TComponentTypeID<int32> Priority;
 		TComponentTypeID<TInstancedStruct<FMovieSceneMixedAnimationTarget>> Target;
 		TComponentTypeID<TSharedPtr<FAnimNextEvaluationTask>> Task;
 		TComponentTypeID<TSharedPtr<FAnimNextEvaluationTask>> MixerTask;
-		TComponentTypeID<ESwapRootBone> RootMotionType;
+		TComponentTypeID<FMovieSceneRootMotionSettings> RootMotionSettings;
 		TComponentTypeID<FObjectComponent> MeshComponent;
+		TComponentTypeID<TSharedPtr<FMovieSceneMixerRootMotionComponentData>> MixerRootMotion;
 
 		struct
 		{
diff --git a/Engine/Plugins/Experimental/MovieSceneAnimMixer/Source/MovieSceneAnimMixer/Public/MovieSceneAnimMixerModule.h b/Engine/Plugins/Experimental/MovieSceneAnimMixer/Source/MovieSceneAnimMixer/Public/MovieSceneAnimMixerModule.h
new file mode 100644
--- /dev/null
+++ b/Engine/Plugins/Experimental/MovieSceneAnimMixer/Source/MovieSceneAnimMixer/Public/MovieSceneAnimMixerModule.h
@@ -0,0 +1,28 @@
+// Copyright Epic Games, Inc. All Rights Reserved.
+
+#pragma once
+
+#include "CoreMinimal.h"
+#include "Modules/ModuleInterface.h"
+
+struct FAnimationUpdateContext;
+
+namespace UE::AnimNext
+{
+	struct FEvaluationVM;
+}
+
+DECLARE_MULTICAST_DELEGATE_TwoParams(FPreEvaluateMixerTasks, const FAnimationUpdateContext&, UE::AnimNext::FEvaluationVM&);
+
+namespace UE::MovieScene
+{
+
+
+class FMovieSceneAnimMixerModule : public IModuleInterface
+{
+	// IModuleInterface interface
+	virtual void StartupModule() override;
+	virtual void ShutdownModule() override;
+};
+
+}
diff --git a/Engine/Plugins/Experimental/MovieSceneAnimMixer/Source/MovieSceneAnimMixer/Public/Systems/MovieSceneAnimMixerSystem.h b/Engine/Plugins/Experimental/MovieSceneAnimMixer/Source/MovieSceneAnimMixer/Public/Systems/MovieSceneAnimMixerSystem.h
--- a/Engine/Plugins/Experimental/MovieSceneAnimMixer/Source/MovieSceneAnimMixer/Public/Systems/MovieSceneAnimMixerSystem.h
+++ b/Engine/Plugins/Experimental/MovieSceneAnimMixer/Source/MovieSceneAnimMixer/Public/Systems/MovieSceneAnimMixerSystem.h
@@ -9,33 +9,28 @@
 #include "Containers/Map.h"
 #include "Templates/SharedPointer.h"
 #include "AnimSequencerInstanceProxy.h"
+#include "Systems/MovieSceneRootMotionSystem.h"
 #include "MovieSceneAnimMixerSystem.generated.h"
 
 struct FAnimNextEvaluationTask;
 
 namespace UE
 {
-	namespace MovieScene
-	{
-		struct EntityID;
-	}
 	namespace AnimNext
 	{
 		struct FEvaluationProgram;
 	}
 }
 
-
 struct FMovieSceneAnimMixerEntry
 {
 	UE::MovieScene::FMovieSceneEntityID EntityID;
 	TSharedPtr<FAnimNextEvaluationTask> EvalTask;
-	double PoseWeight;
-	int32 Priority;
+	double PoseWeight = 1.0;
+	int32 Priority = 0;
 	bool bAdditive = false;
 	bool bRequiresBlend = true;
-	ESwapRootBone RootMotionType = ESwapRootBone::SwapRootBone_None;
-
+	TOptional<FMovieSceneRootMotionSettings> RootMotionSettings;
 	FORCEINLINE bool operator<(const FMovieSceneAnimMixerEntry& RHS) const
 	{
 		return Priority < RHS.Priority;
@@ -69,6 +64,7 @@ struct FMovieSceneAnimMixer
 	UE::MovieScene::FMovieSceneEntityID MixerEntityID;
 	TArray<FMovieSceneAnimMixerEntry> Entries;
 	TSharedPtr<UE::AnimNext::FEvaluationProgram> EvaluationProgram;
+	TSharedPtr<FMovieSceneMixerRootMotionComponentData> RootMotion;
 };
 
 // Takes in evaluation tasks with optional pose weight, masks, priority and a given animation target.
diff --git a/Engine/Plugins/Experimental/MovieSceneAnimMixer/Source/MovieSceneAnimMixer/Public/Systems/MovieSceneAnimNextTargetSystem.h b/Engine/Plugins/Experimental/MovieSceneAnimMixer/Source/MovieSceneAnimMixer/Public/Systems/MovieSceneAnimNextTargetSystem.h
--- a/Engine/Plugins/Experimental/MovieSceneAnimMixer/Source/MovieSceneAnimMixer/Public/Systems/MovieSceneAnimNextTargetSystem.h
+++ b/Engine/Plugins/Experimental/MovieSceneAnimMixer/Source/MovieSceneAnimMixer/Public/Systems/MovieSceneAnimNextTargetSystem.h
@@ -25,7 +25,7 @@ struct FMovieSceneAnimMixerEvaluationModifier : public UE::AnimNext::IEvaluation
 
 	virtual void PostEvaluate(UE::AnimNext::FEvaluateTraversalContext& Context) const override;
 
-	TSharedPtr<FAnimNextEvaluationTask> TaskToInject;
+	TSharedPtr<FAnimNextEvaluationTask> TaskToInject = nullptr;
 
 };
 
@@ -50,10 +50,10 @@ struct MOVIESCENEANIMMIXER_API FMovieSceneAnimNextInjectionTarget : public FMovi
 
 struct FMovieSceneAnimNextTargetData
 {
-	TWeakObjectPtr<UAnimNextComponent> AnimNextComponent;
-	FName InjectionSiteName;
-	UE::AnimNext::FInjectionRequestPtr InjectionRequestHandle;
-	TSharedPtr<FMovieSceneAnimMixerEvaluationModifier> Modifier;
+	TWeakObjectPtr<UAnimNextComponent> AnimNextComponent = nullptr;
+	FName InjectionSiteName = FName();
+	UE::AnimNext::FInjectionRequestPtr InjectionRequestHandle = nullptr;
+	TSharedPtr<FMovieSceneAnimMixerEvaluationModifier> Modifier = nullptr;
 };
 
 // System that handles applying animation mixer evaluation tasks to an injection site in an Anim Next module
diff --git a/Engine/Plugins/Experimental/MovieSceneAnimMixer/Source/MovieSceneAnimMixer/Public/Systems/MovieSceneMixedSkeletalAnimationSystem.h b/Engine/Plugins/Experimental/MovieSceneAnimMixer/Source/MovieSceneAnimMixer/Public/Systems/MovieSceneMixedSkeletalAnimationSystem.h
--- a/Engine/Plugins/Experimental/MovieSceneAnimMixer/Source/MovieSceneAnimMixer/Public/Systems/MovieSceneMixedSkeletalAnimationSystem.h
+++ b/Engine/Plugins/Experimental/MovieSceneAnimMixer/Source/MovieSceneAnimMixer/Public/Systems/MovieSceneMixedSkeletalAnimationSystem.h
@@ -45,19 +45,19 @@ struct FMixedAnimSkeletalAnimationData
 	TOptional<FSkeletalAnimationRootMotionOverride> RootMotionOverride;
 
 	UPROPERTY()
-	double FromPosition;
+	double FromPosition = -1.0f;
 
 	UPROPERTY()
-	double ToPosition;
+	double ToPosition = -1.0f;
 
 	UPROPERTY()
-	bool bFireNotifies;
+	bool bFireNotifies = true;
 
 	UPROPERTY()
-	ESwapRootBone SwapRootBone;
+	ESwapRootBone SwapRootBone = ESwapRootBone::SwapRootBone_None;
 
 	UPROPERTY()
-	bool bAdditive;
+	bool bAdditive = false;
 };
 
 USTRUCT()
diff --git a/Engine/Plugins/Experimental/MovieSceneAnimMixer/Source/MovieSceneAnimMixer/Public/Systems/MovieSceneRootMotionSystem.h b/Engine/Plugins/Experimental/MovieSceneAnimMixer/Source/MovieSceneAnimMixer/Public/Systems/MovieSceneRootMotionSystem.h
--- a/Engine/Plugins/Experimental/MovieSceneAnimMixer/Source/MovieSceneAnimMixer/Public/Systems/MovieSceneRootMotionSystem.h
+++ b/Engine/Plugins/Experimental/MovieSceneAnimMixer/Source/MovieSceneAnimMixer/Public/Systems/MovieSceneRootMotionSystem.h
@@ -3,8 +3,105 @@
 #pragma once
 
 #include "EntitySystem/MovieSceneEntitySystem.h"
+#include "EntitySystem/MovieSceneEntityIDs.h"
+#include "AnimSequencerInstanceProxy.h"
+#include "EvaluationVM/EvaluationTask.h"
+#include "Tickable.h"
 #include "MovieSceneRootMotionSystem.generated.h"
 
+// Enum describing what space that root motion should be applied in.
+UENUM()
+enum class EMovieSceneRootMotionSpace : uint8
+{
+	// Root motion should be applied in animation space, meaning that it will be applied on top of the blended transform track and transform origin.
+	AnimationSpace,
+	// Root motion should be applied in world space, meaning that it will override any transform track or transform origin.
+	WorldSpace
+};
+
+// Structure used for animation tracks to communicate to the mixer how they would like root motion handled if at all.
+struct FMovieSceneRootMotionSettings
+{
+	// Whether to apply root motion to the component or the actor.
+	ESwapRootBone SwapRootBone = ESwapRootBone::SwapRootBone_None;
+
+	// What space to apply root motion in. Defaults to animation space.
+	EMovieSceneRootMotionSpace RootMotionSpace = EMovieSceneRootMotionSpace::AnimationSpace;
+};
+
+// Structure that is shared between entities for handling the mixer's root motion result. 
+// As this can get read/written from multiple threads, access is threadsafe.
+struct FMovieSceneMixerRootMotionComponentData
+{
+public:
+	TOptional<FTransform> GetSequenceAnimationTransform() const;
+	void SetSequenceAnimationTransform(const FTransform& InTransform);
+
+	TOptional<FQuat> GetInverseMeshToActorRotation() const;
+	void SetInverseMeshToActorRotation(const FQuat& InInverseMeshToActorRotation);
+
+	FTransform GetRootTransform() const;
+	void SetRootTransform(const FTransform& InTransform);
+
+	void Reset();
+
+public:
+		// Where to apply the root motion
+		TWeakObjectPtr<USceneComponent> Target;
+
+		// EntityID for the anim mixer
+		UE::MovieScene::FMovieSceneEntityID MixerEntityID;
+
+private:
+	// Blend output result of transform track calculations
+	TOptional<FTransform> SequenceAnimationTransform;
+
+	// Optional inverse mesh component to actor rotation used to offset any mesh component rotation where applicable.
+	TOptional<FQuat> InverseMeshToActorRotation;
+	
+	// Final mixed root transform
+	FTransform RootTransform;
+
+	mutable FTransactionallySafeRWLock RootMotionLock;
+};
+
+// Task that converts the root motion attribute on the top pose of the pose stack to world space by adding on the blended transform track value.
+USTRUCT()
+struct FAnimNextConvertRootMotionToWorldSpaceTask : public FAnimNextEvaluationTask
+{
+	GENERATED_BODY()
+
+	DECLARE_ANIM_EVALUATION_TASK(FAnimNextConvertRootMotionToWorldSpaceTask)
+
+	FAnimNextConvertRootMotionToWorldSpaceTask() = default;
+	FAnimNextConvertRootMotionToWorldSpaceTask(const TSharedPtr<FMovieSceneMixerRootMotionComponentData>& InRootMotionData);
+
+	static FAnimNextConvertRootMotionToWorldSpaceTask Make(const TSharedPtr<FMovieSceneMixerRootMotionComponentData>& InRootMotionData);
+
+	virtual void Execute(UE::AnimNext::FEvaluationVM& VM) const override;
+
+	TSharedPtr<FMovieSceneMixerRootMotionComponentData> RootMotionData;
+};
+
+// Task that gets the final mixed root transform and stores it in the root motion data for later application.
+// TODO: it's not ideal that we're writing things outside of the animation system during an evaluation task. 
+// Consider refactoring this once we have a way to hook into anim next post execution
+USTRUCT()
+struct FAnimNextStoreRootTransformTask : public FAnimNextEvaluationTask
+{
+	GENERATED_BODY()
+
+	DECLARE_ANIM_EVALUATION_TASK(FAnimNextStoreRootTransformTask)
+
+	FAnimNextStoreRootTransformTask() = default;
+	FAnimNextStoreRootTransformTask(const TSharedPtr<FMovieSceneMixerRootMotionComponentData>& InRootMotionData);
+
+	static FAnimNextStoreRootTransformTask Make(const TSharedPtr<FMovieSceneMixerRootMotionComponentData>& InRootMotionData);
+
+	virtual void Execute(UE::AnimNext::FEvaluationVM& VM) const override;
+
+	TSharedPtr<FMovieSceneMixerRootMotionComponentData> RootMotionData;
+};
 
 // Takes in evaluation tasks on mixers.
 // Mixes just the root motion attributes.
@@ -13,6 +110,7 @@
 UCLASS(MinimalAPI)
 class UMovieSceneRootMotionSystem
 	: public UMovieSceneEntitySystem
+	, public FTickableGameObject
 {
 public:
 
@@ -20,8 +118,40 @@ public:
 
 	UMovieSceneRootMotionSystem(const FObjectInitializer& ObjInit);
 
+	TMap<FObjectKey, TSharedPtr<FMovieSceneMixerRootMotionComponentData>>& GetData() { return RootMotionData; }
+
 private:
 
+	virtual void OnRun(FSystemTaskPrerequisites& InPrerequisites, FSystemSubsequentTasks& Subsequents) override;
 	virtual void OnSchedulePersistentTasks(UE::MovieScene::IEntitySystemScheduler* TaskScheduler) override;
+	void FinishDestroy() override;
+
+	/** FTickableGameObject interface */
+	virtual UWorld* GetTickableGameObjectWorld() const override { return GetWorld(); }
+	virtual bool IsTickableInEditor() const override { return true; }
+	virtual bool IsTickable() const override { return !RootMotionData.IsEmpty(); }
+	virtual TStatId GetStatId() const override { RETURN_QUICK_DECLARE_CYCLE_STAT(UMovieSceneRootMotionSystem, STATGROUP_Tickables); }
+	virtual void Tick(float DeltaSeconds) override { ApplyRootMotion(); }
+
+	// Called by the Tick function to apply root motion after animation evaluation on the game thread
+	void ApplyRootMotion();
+
+	TMap<FObjectKey, TSharedPtr<FMovieSceneMixerRootMotionComponentData>> RootMotionData;
+
+#if WITH_EDITOR
+	void OnLink() override;
+	void OnUnlink() override;
+#endif
+
+#if WITH_EDITORONLY_DATA
+	FDelegateHandle OnBeginActorMovementHandle;
+	FDelegateHandle OnEndActorMovementHandle;
+	bool bActorBeingMoved = false;
+#endif
+
+	friend class FMovieSceneRootMotionHelper;
+
+
+protected:
 
 };
\ No newline at end of file
diff --git a/Engine/Source/Runtime/MovieSceneTracks/Public/Systems/MovieSceneSkeletalAnimationSystem.h b/Engine/Source/Runtime/MovieSceneTracks/Public/Systems/MovieSceneSkeletalAnimationSystem.h
--- a/Engine/Source/Runtime/MovieSceneTracks/Public/Systems/MovieSceneSkeletalAnimationSystem.h
+++ b/Engine/Source/Runtime/MovieSceneTracks/Public/Systems/MovieSceneSkeletalAnimationSystem.h
@@ -102,7 +102,7 @@ public:
 
 	UMovieSceneSkeletalAnimationSystem(const FObjectInitializer& ObjInit);
 
-	static UObject* ResolveSkeletalMeshComponentBinding(UObject* InObject);
+	MOVIESCENETRACKS_API static UObject* ResolveSkeletalMeshComponentBinding(UObject* InObject);
 
 	MOVIESCENETRACKS_API FTransform GetRootMotionOffset(UObject* InObject) const;
 	MOVIESCENETRACKS_API void UpdateRootMotionOffset(UObject* InObject);
