Fix Crash when compiling an AnimBP with a level sequence open.
If the sequence had an animation section using the Custom Anim Instance mixed animation target, compiling the animation blueprint for that object would crash due to the reinstanced blueprint being used.

Now when a blueprint is recompiled, we check to see if it matches the source blueprint for the one linked by the seciton, and reset the linked instance, so we no longer tick it in the editor.

We also check if the cahced anim instance for restore state is valid before applying it back to the mesh.

#rnx

diff --git a/Engine/Plugins/Experimental/MovieSceneAnimMixer/Source/MovieSceneAnimMixer/Private/Systems/MovieSceneAnimInstanceTargetSystem.cpp b/Engine/Plugins/Experimental/MovieSceneAnimMixer/Source/MovieSceneAnimMixer/Private/Systems/MovieSceneAnimInstanceTargetSystem.cpp
--- a/Engine/Plugins/Experimental/MovieSceneAnimMixer/Source/MovieSceneAnimMixer/Private/Systems/MovieSceneAnimInstanceTargetSystem.cpp
+++ b/Engine/Plugins/Experimental/MovieSceneAnimMixer/Source/MovieSceneAnimMixer/Private/Systems/MovieSceneAnimInstanceTargetSystem.cpp
@@ -341,9 +341,10 @@ namespace UE::MovieScene
 				// if we're using same anim blueprint, we don't want to keep reinitializing it. 
 				Component->SetAnimationMode(InOutCachedValue.AnimationMode);
 			}
-			if (InOutCachedValue.CachedAnimInstance.Get())
+			UAnimInstance* PreviousAnimInstance = InOutCachedValue.CachedAnimInstance.Get();
+			if (PreviousAnimInstance && IsValid(PreviousAnimInstance) && PreviousAnimInstance->GetSkelMeshComponent() == Component)
 			{
-				Component->AnimScriptInstance = InOutCachedValue.CachedAnimInstance.Get();
+				Component->AnimScriptInstance = PreviousAnimInstance;
 				InOutCachedValue.CachedAnimInstance.Reset();
 				if (Component->AnimScriptInstance && Component->GetSkeletalMeshAsset() && Component->AnimScriptInstance->CurrentSkeleton != Component->GetSkeletalMeshAsset()->GetSkeleton())
 				{
diff --git a/Engine/Source/Runtime/AnimGraphRuntime/Public/AnimCustomInstanceHelper.h b/Engine/Source/Runtime/AnimGraphRuntime/Public/AnimCustomInstanceHelper.h
--- a/Engine/Source/Runtime/AnimGraphRuntime/Public/AnimCustomInstanceHelper.h
+++ b/Engine/Source/Runtime/AnimGraphRuntime/Public/AnimCustomInstanceHelper.h
@@ -31,6 +31,12 @@
 #include "Engine/World.h"
 #include "SequencerAnimationSupport.h"
 
+
+#if WITH_EDITOR
+#include "Editor.h"
+#include "Engine/Blueprint.h"
+#endif
+
 class FAnimCustomInstanceHelper
 {
 public:
@@ -74,6 +80,16 @@ public:
 		const bool bCreateSequencerInterface = InstanceClassType::StaticClass()->ImplementsInterface(USequencerAnimationSupport::StaticClass());
 		bool bSupportDifferentSourceAnimInstance = false;
 
+#if WITH_EDITOR
+		GEditor->OnBlueprintPreCompile().AddLambda([CurrentAnimInstance, InSkeletalMeshComponent](const UBlueprint* InBlueprint)
+		{
+			if (CurrentAnimInstance && InBlueprint && CurrentAnimInstance->GetClass() == InBlueprint->GeneratedClass)
+			{
+				UnbindFromSkeletalMeshComponent<InstanceClassType>(InSkeletalMeshComponent);
+			}
+		});
+#endif
+
 		if (bCreateSequencerInterface)
 		{
 			InstanceClassType* DefaultObject = InstanceClassType::StaticClass()->template GetDefaultObject<InstanceClassType>();
