Sequencer: Mixer plugin-  Fix issue with root motion on multiple sections.

diff --git a/Engine/Plugins/Experimental/MovieSceneAnimMixer/Source/MovieSceneAnimMixer/Private/Systems/MovieSceneMixedSkeletalAnimationSystem.cpp b/Engine/Plugins/Experimental/MovieSceneAnimMixer/Source/MovieSceneAnimMixer/Private/Systems/MovieSceneMixedSkeletalAnimationSystem.cpp
--- a/Engine/Plugins/Experimental/MovieSceneAnimMixer/Source/MovieSceneAnimMixer/Private/Systems/MovieSceneMixedSkeletalAnimationSystem.cpp
+++ b/Engine/Plugins/Experimental/MovieSceneAnimMixer/Source/MovieSceneAnimMixer/Private/Systems/MovieSceneMixedSkeletalAnimationSystem.cpp
@@ -268,7 +268,7 @@ void FMovieSceneSkeletalAnimationEvaluationTask::Execute(UE::AnimNext::FEvaluati
 				}					
 				if (PoseIndex != INDEX_NONE)
 				{
-					Keyframe.Pose.LocalTransformsView[PoseIndex] = AnimationData.RootMotionOverride->RootMotion * AnimationData.RootMotionOverride->PreviousTransform;
+					Keyframe.Pose.LocalTransformsView[PoseIndex] = AnimationData.RootMotionOverride->RootMotion;
 				}
 			}
 
@@ -311,7 +311,7 @@ void FMovieSceneSkeletalAnimationEvaluationTask::Execute(UE::AnimNext::FEvaluati
 			{
 				if (AnimationData.RootMotionOverride.IsSet())
 				{
-					RootBoneTransform = AnimationData.RootMotionOverride->RootMotion * AnimationData.RootMotionOverride->PreviousTransform;
+					RootBoneTransform = AnimationData.RootMotionOverride->RootMotion;
 				}
 				else
 				{
