Sequencer: Tweak to root motion system to ensure that root motion has actually been written before we start setting the character's position, otherwise we can get pops to the origin. I believe part of this is coming from the async nature of anim next evaluation- we may need to deal with tick dependencies better later, but at least this fixes the immediate popping issue I'm seeing.

diff --git a/Engine/Plugins/Experimental/MovieSceneAnimMixer/Source/MovieSceneAnimMixer/Private/Systems/MovieSceneRootMotionSystem.cpp b/Engine/Plugins/Experimental/MovieSceneAnimMixer/Source/MovieSceneAnimMixer/Private/Systems/MovieSceneRootMotionSystem.cpp
--- a/Engine/Plugins/Experimental/MovieSceneAnimMixer/Source/MovieSceneAnimMixer/Private/Systems/MovieSceneRootMotionSystem.cpp
+++ b/Engine/Plugins/Experimental/MovieSceneAnimMixer/Source/MovieSceneAnimMixer/Private/Systems/MovieSceneRootMotionSystem.cpp
@@ -179,12 +179,25 @@ void FMovieSceneMixerRootMotionComponentData::SetInverseMeshToActorRotation(cons
 	InverseMeshToActorRotation = InInverseMeshToActorRotation;
 }
 
-FTransform FMovieSceneMixerRootMotionComponentData::GetRootTransform() const
+TOptional<FTransform> FMovieSceneMixerRootMotionComponentData::GetRootTransform() const
 {
 	UE::TReadScopeLock ScopeLock(RootMotionLock);
 	return RootTransform;
 }
 
+bool FMovieSceneMixerRootMotionComponentData::ConsumeRootTransform(FTransform& OutTransform)
+{
+	UE::TWriteScopeLock ScopeLock(RootMotionLock);
+	if (RootTransform.IsSet())
+	{
+		OutTransform = RootTransform.GetValue();
+		RootTransform.Reset();
+		return true;
+	}
+	return false;
+}
+
+
 void FMovieSceneMixerRootMotionComponentData::SetRootTransform(const FTransform& InTransform)
 {
 	UE::TWriteScopeLock ScopeLock(RootMotionLock);
@@ -197,7 +210,7 @@ void FMovieSceneMixerRootMotionComponentData::Reset()
 	UE::TWriteScopeLock ScopeLock(RootMotionLock);
 	SequenceAnimationTransform.Reset();
 	InverseMeshToActorRotation.Reset();
-	RootTransform.SetIdentity();
+	RootTransform.Reset();
 }
 
 FAnimNextConvertRootMotionToWorldSpaceTask::FAnimNextConvertRootMotionToWorldSpaceTask(const TSharedPtr<FMovieSceneMixerRootMotionComponentData>& InRootMotionData)
@@ -447,8 +460,11 @@ void UMovieSceneRootMotionSystem::ApplyRootMotion()
 		{
 			if (RootMotion->Target.IsValid())
 			{
-				FTransform RelativeTransform = RootMotion->GetRootTransform();
-				RootMotion->Target->SetRelativeLocationAndRotation(RelativeTransform.GetLocation(), RelativeTransform.GetRotation().Rotator());
+				FTransform RelativeTransform;
+				if (RootMotion->ConsumeRootTransform(RelativeTransform))
+				{
+					RootMotion->Target->SetRelativeLocationAndRotation(RelativeTransform.GetLocation(), RelativeTransform.GetRotation().Rotator());
+				}
 			}
 		}
 	}
diff --git a/Engine/Plugins/Experimental/MovieSceneAnimMixer/Source/MovieSceneAnimMixer/Public/Systems/MovieSceneRootMotionSystem.h b/Engine/Plugins/Experimental/MovieSceneAnimMixer/Source/MovieSceneAnimMixer/Public/Systems/MovieSceneRootMotionSystem.h
--- a/Engine/Plugins/Experimental/MovieSceneAnimMixer/Source/MovieSceneAnimMixer/Public/Systems/MovieSceneRootMotionSystem.h
+++ b/Engine/Plugins/Experimental/MovieSceneAnimMixer/Source/MovieSceneAnimMixer/Public/Systems/MovieSceneRootMotionSystem.h
@@ -40,7 +40,8 @@ public:
 	TOptional<FQuat> GetInverseMeshToActorRotation() const;
 	void SetInverseMeshToActorRotation(const FQuat& InInverseMeshToActorRotation);
 
-	FTransform GetRootTransform() const;
+	TOptional<FTransform> GetRootTransform() const;
+	bool ConsumeRootTransform(FTransform& OutTransform);
 	void SetRootTransform(const FTransform& InTransform);
 
 	void Reset();
@@ -60,7 +61,7 @@ private:
 	TOptional<FQuat> InverseMeshToActorRotation;
 	
 	// Final mixed root transform
-	FTransform RootTransform;
+	TOptional<FTransform> RootTransform;
 
 	mutable FTransactionallySafeRWLock RootMotionLock;
 };
