Fix crash in OnPreCompiled callback.

Moving callback registration and handling to teh Anim Instance system. The System will keep an array of handles and unregister them when it is destroyed.

diff --git a/Engine/Plugins/Experimental/MovieSceneAnimMixer/Source/MovieSceneAnimMixer/Private/Systems/MovieSceneAnimInstanceTargetSystem.cpp b/Engine/Plugins/Experimental/MovieSceneAnimMixer/Source/MovieSceneAnimMixer/Private/Systems/MovieSceneAnimInstanceTargetSystem.cpp
--- a/Engine/Plugins/Experimental/MovieSceneAnimMixer/Source/MovieSceneAnimMixer/Private/Systems/MovieSceneAnimInstanceTargetSystem.cpp
+++ b/Engine/Plugins/Experimental/MovieSceneAnimMixer/Source/MovieSceneAnimMixer/Private/Systems/MovieSceneAnimInstanceTargetSystem.cpp
@@ -395,12 +395,15 @@ namespace UE::MovieScene
 
 		UMovieSceneEntitySystemLinker* Linker;
 
+		UMovieSceneAnimInstanceTargetSystem* System;
+
 		TSharedPtr<FPreAnimatedAnimInstanceMixedSkeletalAnimationStorage> PreAnimatedStorage;
 
 	public:
 
-		FEvaluateAndApplyAnimationTasks(UMovieSceneEntitySystemLinker* InLinker)
+		FEvaluateAndApplyAnimationTasks(UMovieSceneEntitySystemLinker* InLinker, UMovieSceneAnimInstanceTargetSystem* InSystem)
 			: Linker(InLinker)
+			, System(InSystem)
 		{
 			PreAnimatedStorage = InLinker->PreAnimatedState.GetOrCreateStorage<FPreAnimatedAnimInstanceMixedSkeletalAnimationStorage>();
 		}
@@ -441,10 +444,31 @@ namespace UE::MovieScene
 			// Setup custom anim instance, using the current anim instance as a 'source'.
 			bool bWasCreated = false;
 			
+#if WITH_EDITOR
+			TWeakObjectPtr CurrentAnimInstanceWeakPtr (SkeletalMeshComponent->GetAnimInstance());
+			TWeakObjectPtr SkeletalMeshComponentWeakPtr (SkeletalMeshComponent);
+#endif
+			
 			USequencerMixedAnimInstance* SequencerInstance = FAnimCustomInstanceHelper::BindToSkeletalMeshComponent<USequencerMixedAnimInstance>(SkeletalMeshComponent, bWasCreated);
 
 			SequencerInstance->SetMixerTask(MixerTask);
 
+#if WITH_EDITOR
+			if (GEditor && bWasCreated)
+			{
+				FDelegateHandle NewDelegateHandle = GEditor->OnBlueprintPreCompile().AddLambda([CurrentAnimInstanceWeakPtr, SkeletalMeshComponentWeakPtr](const UBlueprint* InBlueprint)
+				{
+					if (CurrentAnimInstanceWeakPtr.IsValid() && SkeletalMeshComponentWeakPtr.IsValid() && InBlueprint && CurrentAnimInstanceWeakPtr->GetClass() == InBlueprint->GeneratedClass)
+					{
+						FAnimCustomInstanceHelper::UnbindFromSkeletalMeshComponent<USequencerMixedAnimInstance>(SkeletalMeshComponentWeakPtr.Get());
+					}
+				});
+			
+				System->Handles.Add(NewDelegateHandle);
+			}
+#endif
+			
+
 			// TODO: Figure out- can we do motion vector sim here with the blended anim? Or do we need to do something different
 			/*if (InSkeletalAnimations.SimulatedAnimations.Num() != 0)
 			{
@@ -495,6 +519,19 @@ UMovieSceneAnimInstanceTargetSystem::UMovieSceneAnimInstanceTargetSystem(const F
 	}
 }
 
+#if WITH_EDITOR
+UMovieSceneAnimInstanceTargetSystem::~UMovieSceneAnimInstanceTargetSystem()
+{
+	if (GEditor)
+	{
+		for (FDelegateHandle Handle : Handles)
+		{
+			GEditor->OnBlueprintPreCompile().Remove(Handle);
+		}
+	}
+}
+#endif
+
 void UMovieSceneAnimInstanceTargetSystem::OnSchedulePersistentTasks(UE::MovieScene::IEntitySystemScheduler* TaskScheduler)
 {
 	using namespace UE::MovieScene;
@@ -513,5 +550,5 @@ void UMovieSceneAnimInstanceTargetSystem::OnSchedulePersistentTasks(UE::MovieSce
 		.Read(AnimMixerComponents->MixerTask)
 		.SetParams(Params)
 		.Schedule_PerEntity<FEvaluateAndApplyAnimationTasks>(&Linker->EntityManager, TaskScheduler,
-			Linker);
-}
\ No newline at end of file
+			Linker, this);
+}
diff --git a/Engine/Plugins/Experimental/MovieSceneAnimMixer/Source/MovieSceneAnimMixer/Public/Systems/MovieSceneAnimInstanceTargetSystem.h b/Engine/Plugins/Experimental/MovieSceneAnimMixer/Source/MovieSceneAnimMixer/Public/Systems/MovieSceneAnimInstanceTargetSystem.h
--- a/Engine/Plugins/Experimental/MovieSceneAnimMixer/Source/MovieSceneAnimMixer/Public/Systems/MovieSceneAnimInstanceTargetSystem.h
+++ b/Engine/Plugins/Experimental/MovieSceneAnimMixer/Source/MovieSceneAnimMixer/Public/Systems/MovieSceneAnimInstanceTargetSystem.h
@@ -122,6 +122,12 @@ public:
 
 	UMovieSceneAnimInstanceTargetSystem(const FObjectInitializer& ObjInit);
 
+#if WITH_EDITOR
+	virtual ~UMovieSceneAnimInstanceTargetSystem() override;
+
+	TArray<FDelegateHandle> Handles;
+#endif
+	
 private:
 
 	virtual void OnSchedulePersistentTasks(UE::MovieScene::IEntitySystemScheduler* TaskScheduler) override;
diff --git a/Engine/Source/Runtime/AnimGraphRuntime/Public/AnimCustomInstanceHelper.h b/Engine/Source/Runtime/AnimGraphRuntime/Public/AnimCustomInstanceHelper.h
--- a/Engine/Source/Runtime/AnimGraphRuntime/Public/AnimCustomInstanceHelper.h
+++ b/Engine/Source/Runtime/AnimGraphRuntime/Public/AnimCustomInstanceHelper.h
@@ -32,10 +32,6 @@
 #include "SequencerAnimationSupport.h"
 
 
-#if WITH_EDITOR
-#include "Editor.h"
-#include "Engine/Blueprint.h"
-#endif
 
 class FAnimCustomInstanceHelper
 {
@@ -80,16 +76,6 @@ public:
 		const bool bCreateSequencerInterface = InstanceClassType::StaticClass()->ImplementsInterface(USequencerAnimationSupport::StaticClass());
 		bool bSupportDifferentSourceAnimInstance = false;
 
-#if WITH_EDITOR
-		GEditor->OnBlueprintPreCompile().AddLambda([CurrentAnimInstance, InSkeletalMeshComponent](const UBlueprint* InBlueprint)
-		{
-			if (CurrentAnimInstance && InBlueprint && CurrentAnimInstance->GetClass() == InBlueprint->GeneratedClass)
-			{
-				UnbindFromSkeletalMeshComponent<InstanceClassType>(InSkeletalMeshComponent);
-			}
-		});
-#endif
-
 		if (bCreateSequencerInterface)
 		{
 			InstanceClassType* DefaultObject = InstanceClassType::StaticClass()->template GetDefaultObject<InstanceClassType>();
