Sequencer: Fixed issues with the Anim Mixer and Root Motion animations
Root motion is always manually extracted by FMovieSceneSkeletalAnimationEvaluationTask if the source animation has root motion. This fixes issues where root motion ended up being erroneously baked into the pose. Targets that require root motion to be baked in (such as the anim instance target) are able to extract the final blended root motion from the root motion attribute and re-apply it to the final pose.
Anim mixer entries within the same priority level are now always blended using a multi-way blend (weighted average) and considered to contribute a total weight of 1.f to the priority level. This fixes issues with existing content where multiple weighted anim clips can exist at the same priority level and were previously being blended using multiple two-way blends, rather than a single multi-way blend.

diff --git a/Engine/Plugins/Experimental/MovieSceneAnimMixer/Source/MovieSceneAnimMixer/Private/Systems/MovieSceneAnimInstanceTargetSystem.cpp b/Engine/Plugins/Experimental/MovieSceneAnimMixer/Source/MovieSceneAnimMixer/Private/Systems/MovieSceneAnimInstanceTargetSystem.cpp
--- a/Engine/Plugins/Experimental/MovieSceneAnimMixer/Source/MovieSceneAnimMixer/Private/Systems/MovieSceneAnimInstanceTargetSystem.cpp
+++ b/Engine/Plugins/Experimental/MovieSceneAnimMixer/Source/MovieSceneAnimMixer/Private/Systems/MovieSceneAnimInstanceTargetSystem.cpp
@@ -96,6 +96,11 @@ void USequencerMixedAnimInstance::SetMixerTask(TSharedPtr<FAnimNextEvaluationTas
 	GetProxyOnGameThread<FSequencerMixedAnimInstanceProxy>().SetMixerTask(InEvalTask);
 }
 
+void USequencerMixedAnimInstance::SetRootMotion(const TSharedPtr<FMovieSceneMixerRootMotionComponentData>& InRootMotionData)
+{
+	GetProxyOnGameThread<FSequencerMixedAnimInstanceProxy>().SetRootMotion(InRootMotionData);
+}
+
 void FSequencerMixedAnimInstanceProxy::Initialize(UAnimInstance* InAnimInstance)
 {
 	FAnimInstanceProxy::Initialize(InAnimInstance);
@@ -135,6 +140,8 @@ void FSequencerMixedAnimInstanceProxy::PreEvaluateAnimation(UAnimInstance* InAni
 
 bool FSequencerMixedAnimInstanceProxy::Evaluate(FPoseContext& Output)
 {
+	using namespace UE::MovieScene;
+
 	bool bHasValidOutput = false;
 	if (CurrentSourceProxy)
 	{
@@ -202,6 +209,15 @@ bool FSequencerMixedAnimInstanceProxy::Evaluate(FPoseContext& Output)
 				ResultPose.LODPose.CopyFrom(EvaluatedKeyframe->Pose);
 				ResultPose.Curves.CopyFrom(EvaluatedKeyframe->Curves);
 				ResultPose.Attributes.CopyFrom(EvaluatedKeyframe->Attributes);
+
+				// Re-bake the root motion transform into the pose if we don't have any other way of applying root motion
+				const FTransformAnimationAttribute* RootMotionAttribute = ResultPose.Attributes.Find<FTransformAnimationAttribute>(FAnimMixerComponentTypes::Get()->RootTransformAttributeId);
+				const bool bHasSwapRootBone = RootMotion != nullptr;
+				if (RootMotionAttribute && !bHasSwapRootBone)
+				{
+					ResultPose.LODPose.LocalTransformsView[0] = RootMotionAttribute->Value;
+				}
+
 				bHasValidOutput = true;
 			}
 
@@ -279,12 +295,16 @@ void FSequencerMixedAnimInstanceProxy::SetSourceAnimInstance(UAnimInstance* Sour
 	}
 }
 
-
 void FSequencerMixedAnimInstanceProxy::SetMixerTask(TSharedPtr<FAnimNextEvaluationTask> InEvalTask)
 {
 	MixerTask = InEvalTask;
 }
 
+void FSequencerMixedAnimInstanceProxy::SetRootMotion(const TSharedPtr<FMovieSceneMixerRootMotionComponentData>&  InRootMotion)
+{
+	RootMotion = InRootMotion;
+}
+
 namespace UE::MovieScene
 {
 	/** Helper function to get our sequencer animation instance from a skeletal mesh component */
@@ -413,7 +433,8 @@ namespace UE::MovieScene
 			FMovieSceneEntityID EntityID,
 			FObjectComponent MeshComponent,
 			const TInstancedStruct<FMovieSceneMixedAnimationTarget>& Target,
-			TSharedPtr<FAnimNextEvaluationTask> MixerTask
+			const TSharedPtr<FAnimNextEvaluationTask>& MixerTask,
+			const TSharedPtr<FMovieSceneMixerRootMotionComponentData>* RootMotionData
 		) const
 		{
 			// Invalid or incompatible target, return.
@@ -452,6 +473,7 @@ namespace UE::MovieScene
 			USequencerMixedAnimInstance* SequencerInstance = FAnimCustomInstanceHelper::BindToSkeletalMeshComponent<USequencerMixedAnimInstance>(SkeletalMeshComponent, bWasCreated);
 
 			SequencerInstance->SetMixerTask(MixerTask);
+			SequencerInstance->SetRootMotion(RootMotionData ? *RootMotionData : TSharedPtr<FMovieSceneMixerRootMotionComponentData>());
 
 #if WITH_EDITOR
 			if (GEditor && bWasCreated)
@@ -548,6 +570,7 @@ void UMovieSceneAnimInstanceTargetSystem::OnSchedulePersistentTasks(UE::MovieSce
 		.Read(AnimMixerComponents->MeshComponent)
 		.Read(AnimMixerComponents->Target)
 		.Read(AnimMixerComponents->MixerTask)
+		.ReadOptional(AnimMixerComponents->MixerRootMotion)
 		.SetParams(Params)
 		.Schedule_PerEntity<FEvaluateAndApplyAnimationTasks>(&Linker->EntityManager, TaskScheduler,
 			Linker, this);
diff --git a/Engine/Plugins/Experimental/MovieSceneAnimMixer/Source/MovieSceneAnimMixer/Private/Systems/MovieSceneAnimMixerSystem.cpp b/Engine/Plugins/Experimental/MovieSceneAnimMixer/Source/MovieSceneAnimMixer/Private/Systems/MovieSceneAnimMixerSystem.cpp
--- a/Engine/Plugins/Experimental/MovieSceneAnimMixer/Source/MovieSceneAnimMixer/Private/Systems/MovieSceneAnimMixerSystem.cpp
+++ b/Engine/Plugins/Experimental/MovieSceneAnimMixer/Source/MovieSceneAnimMixer/Private/Systems/MovieSceneAnimMixerSystem.cpp
@@ -19,6 +19,9 @@
 #include "EvaluationVM/Tasks/ExecuteProgram.h"
 #include "EvaluationVM/Tasks/BlendKeyframes.h"
 #include "EvaluationVM/Tasks/ApplyAdditiveKeyframe.h"
+#include "EvaluationVM/Tasks/NormalizeRotations.h"
+#include "EvaluationVM/EvaluationVM.h"
+#include "EvaluationVM/KeyframeState.h"
 #include "EvaluationVM/EvaluationProgram.h"
 #include "Algo/ForEach.h"
 #include "Algo/MaxElement.h"
@@ -93,16 +96,80 @@ namespace UE::MovieScene
 				TSharedPtr<FAnimNextExecuteProgramTask> EvalProgramTask = StaticCastSharedPtr<FAnimNextExecuteProgramTask>(MixerTask);
 				EvalProgramTask->Program = Mixer->EvaluationProgram;
 
-				// From lowest to highest priority, add the tasks to the program.
-				// Each task either
-				// a) Is self-contained, adding a single keyframe to the stack, in which case it will require a blend accumulation task afterwards
-				// b) Relies on the already accumulated pose and outputs a fully mixed pose to the stack, in which case no additional blend accumulation task is required
+				// From lowest to highest priority, add the tasks to the program, grouped by their priority.
+				// Tasks within the same priority are blending using a weighted average, which is then collapsed into a single pose, and blended with the
+				//      subsequnt priority based on its total accumulated PoseWeight.
+				//
+				// For example, 5 animations with a structure of:
+				// Priority 0: (A, w:0.5), (B, w:0.5), (C, w:0.5)
+				// Priority 2: (D, w:0.25), (E, w:0.5)
+				//
+				// Would result in a blend stack equivalent to the following operation:
+				// P1 = A*(0.5/1.5) + B*(0.5/1.5) + C(0.5/1.5)
+				// P2 = D*(0.25/0.75) + B*(0.5/0.75)
+				//
+				// FinalPose = (0.25*P1) + (0.75*P2)
+				//
 				// Note that we assume here that the target will push some kind of 'base pose' to the VM before evaluating the task. This allows us to blend in/out from gameplay for example.
 
-				for (const FMovieSceneAnimMixerEntry& MixerEntry : Mixer->Entries)
+				int32 LastAbsoluteBlendPriority = MIN_int32;
+				float AccumulatedAbsoluteBlendPriortyGroupWeight = 0.f;
+
+				const int32 NumMixerEntries = Mixer->Entries.Num();
+
+				// Function that performs a lookahead accumulation of the total number of absolute weights for the specified start index's priority,
+				//  and returns a value indicating whether the entries need a separate blend stack (true) or whether they can be blended as a single operation
+				auto LookaheadComputeAccumulatedWeight = [&Mixer, NumMixerEntries](int32 StartAtIndex, int32 Priority, float& OutAccumulatedWeight) -> bool
+				{
+					int32 NumAbsoluteBlends = 0;
+					float AccumulatedWeight = 0.f;
+					for (int32 NextIndex = StartAtIndex; NextIndex < NumMixerEntries; ++NextIndex)
+					{
+						const FMovieSceneAnimMixerEntry& NextMixerEntry = Mixer->Entries[NextIndex];
+						if (NextMixerEntry.Priority != Priority || NextMixerEntry.bAdditive)
+						{
+							break;
+						}
+						else if (NextMixerEntry.bRequiresBlend && !NextMixerEntry.bAdditive)
+						{
+							++NumAbsoluteBlends;
+							AccumulatedWeight += NextMixerEntry.PoseWeight;
+						}
+					}
+
+					if (NumAbsoluteBlends > 1 && AccumulatedWeight != 0.0)
+					{
+						OutAccumulatedWeight = AccumulatedWeight;
+						return true;
+					}
+					return false;
+				};
+
+				for (int32 Index = 0; Index < NumMixerEntries; )
 				{
-					if (MixerEntry.EvalTask.IsValid())
+					const int32 Priority = Mixer->Entries[Index].Priority;
+
+					// Recompute the accumulated weights for this priority group
+					float AccumulatedPriorityWeight = 1.f;
+					// If there is more than one pose in this priority with an absolute weight, we need to blend those together using a weighted average before applying additives and blending with the next priority or the base pose
+					const bool bNeedsSeparateWeightStack = LookaheadComputeAccumulatedWeight(Index, Priority, AccumulatedPriorityWeight);
+
+					bool bIsFirstAbsoluteBlend = true;
+					for ( ; Index < NumMixerEntries; ++Index)
 					{
+						const FMovieSceneAnimMixerEntry& MixerEntry = Mixer->Entries[Index];
+						if (MixerEntry.Priority != Priority)
+						{
+							// Break and continue with the outer loop
+							break;
+						}
+
+						if (!MixerEntry.EvalTask.IsValid())
+						{
+							continue;
+						}
+
+						// Evaluate the pose itself
 						Mixer->EvaluationProgram->AppendTaskPtr(MixerEntry.EvalTask);
 
 						if (Mixer->RootMotion && MixerEntry.RootMotionSettings && MixerEntry.RootMotionSettings->RootMotionSpace == EMovieSceneRootMotionSpace::AnimationSpace)
@@ -118,13 +185,37 @@ namespace UE::MovieScene
 								// Add additive blend task.
 								Mixer->EvaluationProgram->AppendTask(FAnimNextApplyAdditiveKeyframeTask::Make(MixerEntry.PoseWeight));
 							}
+							else if (bNeedsSeparateWeightStack)
+							{
+								// Add absolute blend task. The first absolute task within a priority level needs to overwrite the pose with its weighted average,
+								//     subsequent poses are added to this. This in effect leaves us with a single pose on the stack at the end of this priority level
+								//     that includes the weighted average of all absolutes within this priority level, plus all absolutes
+								if (bIsFirstAbsoluteBlend)
+								{
+									bIsFirstAbsoluteBlend = false;
+									Mixer->EvaluationProgram->AppendTask(FAnimNextBlendOverwriteKeyframeWithScaleTask::Make(MixerEntry.PoseWeight / AccumulatedPriorityWeight));
+								}
+								else
+								{
+									Mixer->EvaluationProgram->AppendTask(FMovieSceneAccumulateAbsoluteBlendTask::Make(MixerEntry.PoseWeight / AccumulatedPriorityWeight));
+								}
+							}
 							else
 							{
-								// Add absolute blend task
 								Mixer->EvaluationProgram->AppendTask(FAnimNextBlendTwoKeyframesTask::Make(MixerEntry.PoseWeight));
 							}
 						}
 					}
+					// If we have a previous amount to blend with the last priority, add the tasks to ensure that happens
+					if (bNeedsSeparateWeightStack)
+					{
+						Mixer->EvaluationProgram->AppendTask(FAnimNextNormalizeKeyframeRotationsTask());
+
+						// For now we always blend weighted averages with the next pose with a weight of 1.
+						//    Ultimately this should be controlled by its own weight within an anim mixer layer track.
+						constexpr float WeightedAverageBlendWeight = 1.f;
+						Mixer->EvaluationProgram->AppendTask(FAnimNextBlendTwoKeyframesTask::Make(WeightedAverageBlendWeight));
+					}
 				}
 
 				if (Mixer->RootMotion)
@@ -139,6 +230,55 @@ namespace UE::MovieScene
 	};
 }
 
+FMovieSceneAccumulateAbsoluteBlendTask FMovieSceneAccumulateAbsoluteBlendTask::Make(float ScaleFactor)
+{
+	FMovieSceneAccumulateAbsoluteBlendTask Task;
+	Task.ScaleFactor = ScaleFactor;
+
+	return Task;
+}
+
+void FMovieSceneAccumulateAbsoluteBlendTask::Execute(UE::AnimNext::FEvaluationVM& VM) const
+{
+	using namespace UE::AnimNext;
+
+	// Pop our top two poses, we'll re-use the top keyframe for our result
+
+	TUniquePtr<FKeyframeState> KeyframeA;
+	if (!VM.PopValue(KEYFRAME_STACK_NAME, KeyframeA))
+	{
+		// We have no inputs, nothing to do
+		return;
+	}
+
+	TUniquePtr<FKeyframeState> KeyframeB;
+	if (!VM.PopValue(KEYFRAME_STACK_NAME, KeyframeB))
+	{
+		// We have a single input, leave it on top of the stack
+		VM.PushValue(KEYFRAME_STACK_NAME, MoveTemp(KeyframeB));
+		return;
+	}
+
+	if (EnumHasAnyFlags(VM.GetFlags(), EEvaluationFlags::Bones))
+	{
+		check(KeyframeA->Pose.GetNumBones() == KeyframeB->Pose.GetNumBones());
+
+		BlendAddWithScale(KeyframeB->Pose.LocalTransforms.GetView(), KeyframeA->Pose.LocalTransforms.GetConstView(), ScaleFactor);
+	}
+
+	if (EnumHasAnyFlags(VM.GetFlags(), EEvaluationFlags::Curves))
+	{
+		KeyframeB->Curves.Accumulate(KeyframeA->Curves, ScaleFactor);
+	}
+
+	if (EnumHasAnyFlags(VM.GetFlags(), EEvaluationFlags::Attributes))
+	{
+		UE::Anim::Attributes::AccumulateAttributes(KeyframeA->Attributes, KeyframeB->Attributes, ScaleFactor, AAT_None);
+	}
+
+	VM.PushValue(KEYFRAME_STACK_NAME, MoveTemp(KeyframeB));
+}
+
 UMovieSceneAnimMixerSystem::UMovieSceneAnimMixerSystem(const FObjectInitializer& ObjInit)
 	: Super(ObjInit)
 {
diff --git a/Engine/Plugins/Experimental/MovieSceneAnimMixer/Source/MovieSceneAnimMixer/Private/Systems/MovieSceneMixedSkeletalAnimationSystem.cpp b/Engine/Plugins/Experimental/MovieSceneAnimMixer/Source/MovieSceneAnimMixer/Private/Systems/MovieSceneMixedSkeletalAnimationSystem.cpp
--- a/Engine/Plugins/Experimental/MovieSceneAnimMixer/Source/MovieSceneAnimMixer/Private/Systems/MovieSceneMixedSkeletalAnimationSystem.cpp
+++ b/Engine/Plugins/Experimental/MovieSceneAnimMixer/Source/MovieSceneAnimMixer/Private/Systems/MovieSceneMixedSkeletalAnimationSystem.cpp
@@ -44,6 +44,13 @@
 namespace UE::MovieScene
 {
 
+float GMovieSceneAnimMixerEnabled = true;
+FAutoConsoleVariableRef CVarMovieSceneAnimMixerEnabled(
+	TEXT("Sequencer.AnimMixer.Enabled"),
+	GMovieSceneAnimMixerEnabled,
+	TEXT("(Default: true) Controls whether the new Anim Mixer library is used for skeletal animation evaluation.\n"),
+	ECVF_Default);
+
 // Adds required anim mixer components onto skeletal animation section entities
 struct FSkeletalAnimMixerMutation : IMovieSceneEntityMutation
 {
@@ -248,6 +255,7 @@ void FMovieSceneSkeletalAnimationEvaluationTask::Execute(UE::AnimNext::FEvaluati
 		FKeyframeState Keyframe = VM.MakeUninitializedKeyframe(bIsAdditive);
 
 		const bool bUseRawData = FDecompressionTools::ShouldUseRawData(AnimSequencePtr, Keyframe.Pose);
+		const bool bUseRootMotion = AnimSequencePtr->bEnableRootMotion || AnimationData.SwapRootBone != ESwapRootBone::SwapRootBone_None;
 
 		TOptional<FTransform> RootBoneTransform;
 
@@ -273,7 +281,7 @@ void FMovieSceneSkeletalAnimationEvaluationTask::Execute(UE::AnimNext::FEvaluati
 			}
 
 			// Extract the root motion if we've chosen that option
-			if (AnimationData.SwapRootBone != ESwapRootBone::SwapRootBone_None)
+			if (bUseRootMotion)
 			{
 				int32 RootIndex = Keyframe.Pose.GetRefPose().GetLODBoneIndexFromSkeletonBoneIndex(0);
 				RootBoneTransform = Keyframe.Pose.LocalTransformsView[RootIndex];
@@ -370,6 +378,11 @@ void UMovieSceneMixedSkeletalAnimationSystem::OnRun(FSystemTaskPrerequisites& In
 {
 	using namespace UE::MovieScene;
 
+	if (!GMovieSceneAnimMixerEnabled)
+	{
+		return;
+	}
+
 	FBuiltInComponentTypes* BuiltInComponents = FBuiltInComponentTypes::Get();
 	FMovieSceneTracksComponentTypes* TrackComponents = FMovieSceneTracksComponentTypes::Get();
 	FAnimMixerComponentTypes* AnimMixerComponents = FAnimMixerComponentTypes::Get();
diff --git a/Engine/Plugins/Experimental/MovieSceneAnimMixer/Source/MovieSceneAnimMixer/Public/Systems/MovieSceneAnimInstanceTargetSystem.h b/Engine/Plugins/Experimental/MovieSceneAnimMixer/Source/MovieSceneAnimMixer/Public/Systems/MovieSceneAnimInstanceTargetSystem.h
--- a/Engine/Plugins/Experimental/MovieSceneAnimMixer/Source/MovieSceneAnimMixer/Public/Systems/MovieSceneAnimInstanceTargetSystem.h
+++ b/Engine/Plugins/Experimental/MovieSceneAnimMixer/Source/MovieSceneAnimMixer/Public/Systems/MovieSceneAnimInstanceTargetSystem.h
@@ -13,6 +13,7 @@
 
 
 struct FAnimNextEvaluationTask;
+struct FMovieSceneMixerRootMotionComponentData;
 
 /**
  * Declaring a unique target for targeting a custom anim instance on the skeletal mesh component. Doesn't need additional metadata.
@@ -39,6 +40,7 @@ class USequencerMixedAnimInstance : public UAnimInstance, public ISequencerAnima
 
 public:
 	MOVIESCENEANIMMIXER_API void SetMixerTask(TSharedPtr<FAnimNextEvaluationTask> InEvalTask);
+	MOVIESCENEANIMMIXER_API void SetRootMotion(const TSharedPtr<FMovieSceneMixerRootMotionComponentData>& InRootMotion);
 
 protected:
 	// UAnimInstance interface
@@ -90,8 +92,7 @@ public:
 	UAnimInstance* GetSourceAnimInstance() const { return CurrentSourceAnimInstance; }
 
 	void SetMixerTask(TSharedPtr<FAnimNextEvaluationTask> InEvalTask);
-
-
+	void SetRootMotion(const TSharedPtr<FMovieSceneMixerRootMotionComponentData>& InRootMotion);
 	void LinkSourcePose(UAnimInstance* InInputInstance, FAnimInstanceProxy* InInputProxy);
 	void UnlinkSourcePose();
 
@@ -106,6 +107,9 @@ protected:
 
 	// Pointer to the task passed from the mixer.
 	TSharedPtr<FAnimNextEvaluationTask> MixerTask;
+
+	// Root motion settings from the mixed animation
+	TSharedPtr<FMovieSceneMixerRootMotionComponentData> RootMotion;
 };
 
 
diff --git a/Engine/Plugins/Experimental/MovieSceneAnimMixer/Source/MovieSceneAnimMixer/Public/Systems/MovieSceneAnimMixerSystem.h b/Engine/Plugins/Experimental/MovieSceneAnimMixer/Source/MovieSceneAnimMixer/Public/Systems/MovieSceneAnimMixerSystem.h
--- a/Engine/Plugins/Experimental/MovieSceneAnimMixer/Source/MovieSceneAnimMixer/Public/Systems/MovieSceneAnimMixerSystem.h
+++ b/Engine/Plugins/Experimental/MovieSceneAnimMixer/Source/MovieSceneAnimMixer/Public/Systems/MovieSceneAnimMixerSystem.h
@@ -33,6 +33,15 @@ struct FMovieSceneAnimMixerEntry
 	TOptional<FMovieSceneRootMotionSettings> RootMotionSettings;
 	FORCEINLINE bool operator<(const FMovieSceneAnimMixerEntry& RHS) const
 	{
+		if (Priority == RHS.Priority)
+		{
+			// Sort additives after absolutes so they are applied on top
+			if (bAdditive != RHS.bAdditive)
+			{
+				return RHS.bAdditive;
+			}
+			return false;
+		}
 		return Priority < RHS.Priority;
 	}
 };
@@ -67,6 +76,34 @@ struct FMovieSceneAnimMixer
 	TSharedPtr<FMovieSceneMixerRootMotionComponentData> RootMotion;
 };
 
+
+/*
+ * Sequencer weighted average addition blend task
+ *
+ * This happens to be identical to FAnimNextBlendAddKeyframeWithScaleTask except the operatnds are reversed,
+ *     so the resulting stack state is:
+
+ * Top = (Top-1) + (Top * ScaleFactor)
+ * 
+ * Note that rotations will not be normalized after this task.
+ */
+USTRUCT()
+struct FMovieSceneAccumulateAbsoluteBlendTask : public FAnimNextEvaluationTask
+{
+	GENERATED_BODY()
+
+	DECLARE_ANIM_EVALUATION_TASK(FMovieSceneAccumulateAbsoluteBlendTask)
+
+	static FMovieSceneAccumulateAbsoluteBlendTask Make(float ScaleFactor);
+
+	// Task entry point
+	virtual void Execute(UE::AnimNext::FEvaluationVM& VM) const override;
+
+	// The scale factor to apply to the keyframe
+	UPROPERTY()
+	float ScaleFactor = 0.0f;
+};
+
 // Takes in evaluation tasks with optional pose weight, masks, priority and a given animation target.
 // Constructs a hierarchical 'mixer' per bound object per target. 
 // Similar to blender systems, in a 'many to one' operation, each mixer will create an entity with a single evaluation task
