Sequencer: Prevent repeatedly injecting the same evaluation modifier for AnimNext targets

diff --git a/Engine/Plugins/Experimental/MovieSceneAnimMixer/Source/MovieSceneAnimMixer/Private/Systems/MovieSceneAnimNextTargetSystem.cpp b/Engine/Plugins/Experimental/MovieSceneAnimMixer/Source/MovieSceneAnimMixer/Private/Systems/MovieSceneAnimNextTargetSystem.cpp
--- a/Engine/Plugins/Experimental/MovieSceneAnimMixer/Source/MovieSceneAnimMixer/Private/Systems/MovieSceneAnimNextTargetSystem.cpp
+++ b/Engine/Plugins/Experimental/MovieSceneAnimMixer/Source/MovieSceneAnimMixer/Private/Systems/MovieSceneAnimNextTargetSystem.cpp
@@ -165,24 +165,21 @@ namespace UE::MovieScene
 
 			if (AnimNextComponent)
 			{
+				UE::AnimNext::FInjectionSite InjectionSite(InjectionTarget->InjectionSiteName);
+				if (InjectionSite.DesiredSiteName == NAME_None)
+				{
+					InjectionSite.DesiredSiteName = GSequencerDefaultAnimNextInjectionSite;
+					InjectionSite.bUseModuleFallback = true;
+				}
+
 				// If not currently injected, inject
-				FMovieSceneAnimNextTargetData* TargetData = System->CurrentTargets.FindByPredicate([AnimNextComponent, InjectionTarget](const FMovieSceneAnimNextTargetData& Data) { return Data.AnimNextComponent == AnimNextComponent && Data.InjectionSite.DesiredSiteName == InjectionTarget->InjectionSiteName;});
+				FMovieSceneAnimNextTargetData* TargetData = System->CurrentTargets.FindByPredicate([AnimNextComponent, InjectionSite](const FMovieSceneAnimNextTargetData& Data) { return Data.AnimNextComponent == AnimNextComponent && Data.InjectionSite.DesiredSiteName == InjectionSite.DesiredSiteName;});
 
 				if (!TargetData)
 				{
 					TargetData = &System->CurrentTargets.AddDefaulted_GetRef();
 					TargetData->AnimNextComponent = AnimNextComponent;
-
-					if (InjectionTarget->InjectionSiteName == NAME_None)
-					{
-						TargetData->InjectionSite.DesiredSiteName = GSequencerDefaultAnimNextInjectionSite;
-						TargetData->InjectionSite.bUseModuleFallback = true;
-					}
-					else
-					{
-						TargetData->InjectionSite.DesiredSiteName = InjectionTarget->InjectionSiteName;
-						TargetData->InjectionSite.bUseModuleFallback = false;
-					}
+					TargetData->InjectionSite = InjectionSite;
 					TargetData->Modifier = MakeShared<FMovieSceneAnimMixerEvaluationModifier>(MixerTask);
 
 					// Inject
