Sequencer: Add strong references to graph injection requests

Ultimately all these references should probably be weak object ptrs, but since AnimNext is currently set up to use strong references, we'll follow suit here.

diff --git a/Engine/Plugins/Experimental/AnimNextAnimGraph/Source/AnimNextAnimGraph/Private/Injection/InjectionRequest.cpp b/Engine/Plugins/Experimental/AnimNextAnimGraph/Source/AnimNextAnimGraph/Private/Injection/InjectionRequest.cpp
--- a/Engine/Plugins/Experimental/AnimNextAnimGraph/Source/AnimNextAnimGraph/Private/Injection/InjectionRequest.cpp
+++ b/Engine/Plugins/Experimental/AnimNextAnimGraph/Source/AnimNextAnimGraph/Private/Injection/InjectionRequest.cpp
@@ -241,6 +241,11 @@ namespace UE::AnimNext
 		TimelineState = NewTimelineState;
 	}
 
+	void FInjectionRequest::ExternalAddReferencedObjects(FReferenceCollector& Collector)
+	{
+		AddReferencedObjects(Collector);
+	}
+
 	void FInjectionRequest::AddReferencedObjects(FReferenceCollector& Collector)
 	{
 		Collector.AddPropertyReferencesWithStructARO(FInjectionRequestArgs::StaticStruct(), &RequestArgs);
diff --git a/Engine/Plugins/Experimental/AnimNextAnimGraph/Source/AnimNextAnimGraph/Public/Injection/InjectionRequest.h b/Engine/Plugins/Experimental/AnimNextAnimGraph/Source/AnimNextAnimGraph/Public/Injection/InjectionRequest.h
--- a/Engine/Plugins/Experimental/AnimNextAnimGraph/Source/AnimNextAnimGraph/Public/Injection/InjectionRequest.h
+++ b/Engine/Plugins/Experimental/AnimNextAnimGraph/Source/AnimNextAnimGraph/Public/Injection/InjectionRequest.h
@@ -194,6 +194,8 @@ namespace UE::AnimNext
 		// Returns whether or not this request was interrupted (by Stop or by another request)
 		ANIMNEXTANIMGRAPH_API bool WasInterrupted() const;
 
+		ANIMNEXTANIMGRAPH_API void ExternalAddReferencedObjects(FReferenceCollector& Collector);
+
 	private:
 		// Sends this request to the specified host and it will attempt to play with the requested arguments
 		bool Inject(FInjectionRequestArgs&& InRequestArgs, FInjectionLifetimeEvents&& InLifetimeEvents, UObject* InHost, FModuleHandle InHandle);
diff --git a/Engine/Plugins/Experimental/MovieSceneAnimMixer/Source/MovieSceneAnimMixer/Private/Systems/MovieSceneAnimNextTargetSystem.cpp b/Engine/Plugins/Experimental/MovieSceneAnimMixer/Source/MovieSceneAnimMixer/Private/Systems/MovieSceneAnimNextTargetSystem.cpp
--- a/Engine/Plugins/Experimental/MovieSceneAnimMixer/Source/MovieSceneAnimMixer/Private/Systems/MovieSceneAnimNextTargetSystem.cpp
+++ b/Engine/Plugins/Experimental/MovieSceneAnimMixer/Source/MovieSceneAnimMixer/Private/Systems/MovieSceneAnimNextTargetSystem.cpp
@@ -242,6 +242,19 @@ UMovieSceneAnimNextTargetSystem::UMovieSceneAnimNextTargetSystem(const FObjectIn
 	}
 }
 
+void UMovieSceneAnimNextTargetSystem::AddReferencedObjects(UObject* InThis, FReferenceCollector& Collector)
+{
+	UMovieSceneAnimNextTargetSystem* This = static_cast<UMovieSceneAnimNextTargetSystem*>(InThis);
+
+	for (const FMovieSceneAnimNextTargetData& TargetData : This->CurrentTargets)
+	{
+		if (TargetData.InjectionRequestHandle)
+		{
+			TargetData.InjectionRequestHandle->ExternalAddReferencedObjects(Collector);
+		}
+	}
+}
+
 void UMovieSceneAnimNextTargetSystem::OnUnlink()
 {
 	// Clean up system data
diff --git a/Engine/Plugins/Experimental/MovieSceneAnimMixer/Source/MovieSceneAnimMixer/Public/Systems/MovieSceneAnimNextTargetSystem.h b/Engine/Plugins/Experimental/MovieSceneAnimMixer/Source/MovieSceneAnimMixer/Public/Systems/MovieSceneAnimNextTargetSystem.h
--- a/Engine/Plugins/Experimental/MovieSceneAnimMixer/Source/MovieSceneAnimMixer/Public/Systems/MovieSceneAnimNextTargetSystem.h
+++ b/Engine/Plugins/Experimental/MovieSceneAnimMixer/Source/MovieSceneAnimMixer/Public/Systems/MovieSceneAnimNextTargetSystem.h
@@ -67,6 +67,8 @@ public:
 
 	UMovieSceneAnimNextTargetSystem(const FObjectInitializer& ObjInit);
 
+	MOVIESCENEANIMMIXER_API static void AddReferencedObjects(UObject* InThis, FReferenceCollector& Collector);
+
 	void OnUnlink() override;
 
 	TArray<FMovieSceneAnimNextTargetData> CurrentTargets;
