Sequencer: Add missing call to Super::AddReferencedObjects in UMovieSceneAnimNextTargetSystem

diff --git a/Engine/Plugins/Experimental/MovieSceneAnimMixer/Source/MovieSceneAnimMixer/Private/Systems/MovieSceneAnimNextTargetSystem.cpp b/Engine/Plugins/Experimental/MovieSceneAnimMixer/Source/MovieSceneAnimMixer/Private/Systems/MovieSceneAnimNextTargetSystem.cpp
--- a/Engine/Plugins/Experimental/MovieSceneAnimMixer/Source/MovieSceneAnimMixer/Private/Systems/MovieSceneAnimNextTargetSystem.cpp
+++ b/Engine/Plugins/Experimental/MovieSceneAnimMixer/Source/MovieSceneAnimMixer/Private/Systems/MovieSceneAnimNextTargetSystem.cpp
@@ -246,6 +246,8 @@ void UMovieSceneAnimNextTargetSystem::AddReferencedObjects(UObject* InThis, FRef
 {
 	UMovieSceneAnimNextTargetSystem* This = static_cast<UMovieSceneAnimNextTargetSystem*>(InThis);
 
+	Super::AddReferencedObjects(InThis, Collector);
+
 	for (const FMovieSceneAnimNextTargetData& TargetData : This->CurrentTargets)
 	{
 		if (TargetData.InjectionRequestHandle)
