Sequencer: Adding additional check for null component to workaround GC bug with anim next components

This should not be necessary due to the strong refs that we add in AddReferencedObjects but the components are still getting collected somehow.

diff --git a/Engine/Plugins/Experimental/MovieSceneAnimMixer/Source/MovieSceneAnimMixer/Private/Systems/MovieSceneAnimNextTargetSystem.cpp b/Engine/Plugins/Experimental/MovieSceneAnimMixer/Source/MovieSceneAnimMixer/Private/Systems/MovieSceneAnimNextTargetSystem.cpp
--- a/Engine/Plugins/Experimental/MovieSceneAnimMixer/Source/MovieSceneAnimMixer/Private/Systems/MovieSceneAnimNextTargetSystem.cpp
+++ b/Engine/Plugins/Experimental/MovieSceneAnimMixer/Source/MovieSceneAnimMixer/Private/Systems/MovieSceneAnimNextTargetSystem.cpp
@@ -263,7 +263,8 @@ void UMovieSceneAnimNextTargetSystem::OnUnlink()
 	for (const FMovieSceneAnimNextTargetData& TargetData : CurrentTargets)
 	{
 		// Remove injection
-		if (TargetData.InjectionRequestHandle.IsValid())
+		if (TargetData.AnimNextComponent.Get() != nullptr && // This shouldn't be necessary due to UMovieSceneAnimNextTargetSystem::AddReferencedObjects, but somehow there is still a crash so we're going with the belt-and-braces approach
+			TargetData.InjectionRequestHandle.IsValid())
 		{
 			UE::AnimNext::FInjectionUtils::Uninject(TargetData.InjectionRequestHandle);
 		}
