Sequencer: Root motion computation from skeletal animation sections now happens on the game thread to work around threading constraints

- Ultimately all this code should be happening inside the evaluation system, but the quick fix for now is to put it on the gamethread

diff --git a/Engine/Plugins/Experimental/MovieSceneAnimMixer/Source/MovieSceneAnimMixer/Private/Systems/MovieSceneMixedSkeletalAnimationSystem.cpp b/Engine/Plugins/Experimental/MovieSceneAnimMixer/Source/MovieSceneAnimMixer/Private/Systems/MovieSceneMixedSkeletalAnimationSystem.cpp
--- a/Engine/Plugins/Experimental/MovieSceneAnimMixer/Source/MovieSceneAnimMixer/Private/Systems/MovieSceneMixedSkeletalAnimationSystem.cpp
+++ b/Engine/Plugins/Experimental/MovieSceneAnimMixer/Source/MovieSceneAnimMixer/Private/Systems/MovieSceneMixedSkeletalAnimationSystem.cpp
@@ -11,6 +11,7 @@
 #include "EntitySystem/MovieSceneEntitySystemTask.h"
 #include "MovieSceneTracksComponentTypes.h"
 #include "Sections/MovieSceneSkeletalAnimationSection.h"
+#include "Tracks/MovieSceneSkeletalAnimationTrack.h"
 
 #include "Components/SkeletalMeshComponent.h"
 #include "Engine/SkeletalMesh.h"
@@ -108,11 +109,56 @@ struct FSkeletalAnimMixerMutation : IMovieSceneEntityMutation
 				OutPriorities[Index] = SkeletalAnim.Section->MixedAnimationPriority;
 				OutRootMotionSettings[Index].SwapRootBone = SkeletalAnim.Section->Params.SwapRootBone;
 				OutRootMotionSettings[Index].RootMotionSpace = EMovieSceneRootMotionSpace::AnimationSpace;
+
+				UMovieSceneSkeletalAnimationTrack* Track = SkeletalAnim.Section->GetTypedOuter<UMovieSceneSkeletalAnimationTrack>();
+
+				if (Track->RootMotionParams.bRootMotionsDirty) 
+				{
+					constexpr bool bForce = true;
+					Track->SetUpRootMotions(bForce);
+				}
+
+				if (Track->RootMotionParams.bHaveRootMotion)
+				{
+					OutRootMotionSettings[Index].RootMotionOverride.Emplace();
+				}
+				else
+				{
+					OutRootMotionSettings[Index].RootMotionOverride.Reset();
+				}
+			}
+		}
+	}
+};
+
+
+
+struct FGatherMixableRootMotion
+{
+	static void ForEachEntity(
+		const FMovieSceneSkeletalAnimationComponentData& SkeletalAnimation,
+		FFrameTime EvalTime,
+		FMovieSceneRootMotionSettings& RootMotionSettings
+		)
+	{
+		if (RootMotionSettings.RootMotionOverride.IsSet())
+		{
+			UMovieSceneSkeletalAnimationSection::FRootMotionParams RootMotionParams;
+
+			SkeletalAnimation.Section->GetRootMotion(EvalTime, RootMotionParams);
+			if (RootMotionParams.Transform.IsSet())
+			{
+				RootMotionSettings.RootMotionOverride = FSkeletalAnimationRootMotionOverride();
+				RootMotionSettings.RootMotionOverride->RootMotion = RootMotionParams.Transform.GetValue();
+				RootMotionSettings.RootMotionOverride->bBlendFirstChildOfRoot = RootMotionParams.bBlendFirstChildOfRoot;
+				RootMotionSettings.RootMotionOverride->ChildBoneIndex = RootMotionParams.ChildBoneIndex;
+				RootMotionSettings.RootMotionOverride->PreviousTransform = RootMotionParams.PreviousTransform.GetValue();
 			}
 		}
 	}
 };
 
+
 /** ------------------------------------------------------------------------- */
 /** Task for gathering active skeletal animations and setting up their tasks. */
 struct FGatherMixableSkeletalAnimations
@@ -131,6 +177,7 @@ struct FGatherMixableSkeletalAnimations
 		TRead<UObject*> BoundObjects,
 		TRead<FMovieSceneSkeletalAnimationComponentData> SkeletalAnimations,
 		TReadOptional<FFrameTime> OptionalEvalTimes,
+		TReadOptional<FMovieSceneRootMotionSettings> OptionalRootMotionSettings,
 		TWrite<TSharedPtr<FAnimNextEvaluationTask>> EvalTask) const
 	{
 		// Gather all the skeletal animations currently active in all sequences.
@@ -190,22 +237,12 @@ struct FGatherMixableSkeletalAnimations
 				bLooped = true;
 			}
 
-			TOptional<FSkeletalAnimationRootMotionOverride> RootMotionOverride;
-			UMovieSceneSkeletalAnimationSection::FRootMotionParams RootMotionParams;
-
-			AnimSection->GetRootMotion(EvalFrameTime, RootMotionParams);
-			if (RootMotionParams.Transform.IsSet())
-			{
-				RootMotionOverride = FSkeletalAnimationRootMotionOverride();
-				RootMotionOverride.GetValue().RootMotion = RootMotionParams.Transform.GetValue();
-				RootMotionOverride.GetValue().bBlendFirstChildOfRoot = RootMotionParams.bBlendFirstChildOfRoot;
-				RootMotionOverride.GetValue().ChildBoneIndex = RootMotionParams.ChildBoneIndex;
-				RootMotionOverride.GetValue().PreviousTransform = RootMotionParams.PreviousTransform.GetValue();
-			}
-
 			FMixedAnimSkeletalAnimationData AnimData;
 			AnimData.AnimSequence = AnimSequence;
-			AnimData.RootMotionOverride = RootMotionOverride;
+			if (OptionalRootMotionSettings)
+			{
+				AnimData.RootMotionOverride = OptionalRootMotionSettings[Index].RootMotionOverride;
+			}
 			AnimData.FromPosition = PreviousEvalTime;
 			AnimData.ToPosition = EvalTime;
 			AnimData.bFireNotifies = bFireNotifies && !AnimParams.bSkipAnimNotifiers && !bLooped;
@@ -235,7 +272,7 @@ private:
 		return GIsEditor && bIsNotInPIEOrNotPlaying;
 	}
 };
-	
+
 } // namespace UE::MovieScene
 
 void FMovieSceneSkeletalAnimationEvaluationTask::Execute(UE::AnimNext::FEvaluationVM& VM) const
@@ -408,6 +445,15 @@ void UMovieSceneMixedSkeletalAnimationSystem::OnSchedulePersistentTasks(UE::Movi
 	FMovieSceneTracksComponentTypes* TrackComponents = FMovieSceneTracksComponentTypes::Get();
 	FAnimMixerComponentTypes* AnimMixerComponents = FAnimMixerComponentTypes::Get();
 
+	// Gather all skel anim section entities that have been marked with the AnimMixerPoseProducer tag
+	FTaskID GatherRootMotionTask = FEntityTaskBuilder()
+		.Read(TrackComponents->SkeletalAnimation)
+		.Read(BuiltInComponents->EvalTime)
+		.Write(AnimMixerComponents->RootMotionSettings)
+		.SetDesiredThread(ENamedThreads::GameThread_Local)
+		.FilterNone({ BuiltInComponents->Tags.Ignored })
+		.Schedule_PerEntity<FGatherMixableRootMotion>(&Linker->EntityManager, TaskScheduler);
+
 	// Gather all skel anim section entities that have been marked with the AnimMixerPoseProducer tag
 	FTaskID GatherTask = FEntityTaskBuilder()
 		.ReadEntityIDs()
@@ -415,9 +461,12 @@ void UMovieSceneMixedSkeletalAnimationSystem::OnSchedulePersistentTasks(UE::Movi
 		.Read(BuiltInComponents->BoundObject)
 		.Read(TrackComponents->SkeletalAnimation)
 		.ReadOptional(BuiltInComponents->EvalTime)
+		.ReadOptional(AnimMixerComponents->RootMotionSettings)
 		.Write(AnimMixerComponents->Task)
 		.FilterAll( { TrackComponents->Tags.AnimMixerPoseProducer })
 		.FilterNone({ BuiltInComponents->Tags.Ignored})
 		.Schedule_PerAllocation<FGatherMixableSkeletalAnimations>(&Linker->EntityManager, TaskScheduler,
 			Linker->GetInstanceRegistry());
+
+	TaskScheduler->AddPrerequisite(GatherRootMotionTask, GatherTask);
 }
diff --git a/Engine/Plugins/Experimental/MovieSceneAnimMixer/Source/MovieSceneAnimMixer/Public/Systems/MovieSceneMixedSkeletalAnimationSystem.h b/Engine/Plugins/Experimental/MovieSceneAnimMixer/Source/MovieSceneAnimMixer/Public/Systems/MovieSceneMixedSkeletalAnimationSystem.h
--- a/Engine/Plugins/Experimental/MovieSceneAnimMixer/Source/MovieSceneAnimMixer/Public/Systems/MovieSceneMixedSkeletalAnimationSystem.h
+++ b/Engine/Plugins/Experimental/MovieSceneAnimMixer/Source/MovieSceneAnimMixer/Public/Systems/MovieSceneMixedSkeletalAnimationSystem.h
@@ -9,29 +9,12 @@
 #include "EntitySystem/MovieSceneEntityInstantiatorSystem.h"
 #include "EvaluationVM/EvaluationTask.h"
 #include "AnimSequencerInstanceProxy.h"
+#include "Systems/MovieSceneRootMotionSystem.h"
 #include "MovieSceneMixedSkeletalAnimationSystem.generated.h"
 
 struct FAnimSequencerData;
 class UAnimSequence;
 
-USTRUCT()
-struct FSkeletalAnimationRootMotionOverride
-{
-	GENERATED_BODY()
-
-	/** If true we use the ChildBoneIndex otherwise we use the root*/
-	UPROPERTY()
-	bool bBlendFirstChildOfRoot = false;
-	
-	UPROPERTY()
-	int32 ChildBoneIndex = INDEX_NONE;
-	
-	UPROPERTY()
-	FTransform RootMotion = FTransform::Identity;
-	
-	UPROPERTY()
-	FTransform PreviousTransform = FTransform::Identity;
-};
 
 USTRUCT()
 struct FMixedAnimSkeletalAnimationData
diff --git a/Engine/Plugins/Experimental/MovieSceneAnimMixer/Source/MovieSceneAnimMixer/Public/Systems/MovieSceneRootMotionSystem.h b/Engine/Plugins/Experimental/MovieSceneAnimMixer/Source/MovieSceneAnimMixer/Public/Systems/MovieSceneRootMotionSystem.h
--- a/Engine/Plugins/Experimental/MovieSceneAnimMixer/Source/MovieSceneAnimMixer/Public/Systems/MovieSceneRootMotionSystem.h
+++ b/Engine/Plugins/Experimental/MovieSceneAnimMixer/Source/MovieSceneAnimMixer/Public/Systems/MovieSceneRootMotionSystem.h
@@ -19,6 +19,25 @@ enum class EMovieSceneRootMotionSpace : uint8
 	WorldSpace
 };
 
+USTRUCT()
+struct FSkeletalAnimationRootMotionOverride
+{
+	GENERATED_BODY()
+
+	/** If true we use the ChildBoneIndex otherwise we use the root*/
+	UPROPERTY()
+	bool bBlendFirstChildOfRoot = false;
+	
+	UPROPERTY()
+	int32 ChildBoneIndex = INDEX_NONE;
+	
+	UPROPERTY()
+	FTransform RootMotion = FTransform::Identity;
+	
+	UPROPERTY()
+	FTransform PreviousTransform = FTransform::Identity;
+};
+
 // Structure used for animation tracks to communicate to the mixer how they would like root motion handled if at all.
 struct FMovieSceneRootMotionSettings
 {
@@ -27,6 +46,8 @@ struct FMovieSceneRootMotionSettings
 
 	// What space to apply root motion in. Defaults to animation space.
 	EMovieSceneRootMotionSpace RootMotionSpace = EMovieSceneRootMotionSpace::AnimationSpace;
+
+	TOptional<FSkeletalAnimationRootMotionOverride> RootMotionOverride;
 };
 
 // Structure that is shared between entities for handling the mixer's root motion result. 
