Sequencer: Fixed potential for unsafe threaded access in root motion systems

diff --git a/Engine/Plugins/Experimental/MovieSceneAnimMixer/Source/MovieSceneAnimMixer/Private/Systems/MovieSceneRootMotionSystem.cpp b/Engine/Plugins/Experimental/MovieSceneAnimMixer/Source/MovieSceneAnimMixer/Private/Systems/MovieSceneRootMotionSystem.cpp
--- a/Engine/Plugins/Experimental/MovieSceneAnimMixer/Source/MovieSceneAnimMixer/Private/Systems/MovieSceneRootMotionSystem.cpp
+++ b/Engine/Plugins/Experimental/MovieSceneAnimMixer/Source/MovieSceneAnimMixer/Private/Systems/MovieSceneRootMotionSystem.cpp
@@ -441,8 +441,9 @@ void UMovieSceneRootMotionSystem::OnSchedulePersistentTasks(UE::MovieScene::IEnt
 					BuiltInComponents->DoubleResult[7],
 					BuiltInComponents->DoubleResult[8])
 	.FilterAll({TrackComponents->ComponentTransform.PropertyTag})
+	.FilterNone({BuiltInComponents->BlendChannelInput})
 	.SetDesiredThread(Linker->EntityManager.GetGatherThread())
-	.Fork_PerAllocation<FSyncTransformTrackResult>(&Linker->EntityManager, TaskScheduler, this);
+	.Schedule_PerAllocation<FSyncTransformTrackResult>(&Linker->EntityManager, TaskScheduler, this);
 }
 
 void UMovieSceneRootMotionSystem::ApplyRootMotion()
