Fix Popping issue from floating point errors when blending mixed animations in sequencer

We now skip animaions with blend weights that are nearly zero.

diff --git a/Engine/Plugins/Experimental/MovieSceneAnimMixer/Source/MovieSceneAnimMixer/Private/Systems/MovieSceneAnimMixerSystem.cpp b/Engine/Plugins/Experimental/MovieSceneAnimMixer/Source/MovieSceneAnimMixer/Private/Systems/MovieSceneAnimMixerSystem.cpp
--- a/Engine/Plugins/Experimental/MovieSceneAnimMixer/Source/MovieSceneAnimMixer/Private/Systems/MovieSceneAnimMixerSystem.cpp
+++ b/Engine/Plugins/Experimental/MovieSceneAnimMixer/Source/MovieSceneAnimMixer/Private/Systems/MovieSceneAnimMixerSystem.cpp
@@ -122,6 +122,7 @@ namespace UE::MovieScene
 				auto LookaheadComputeAccumulatedWeight = [&Mixer, NumMixerEntries](int32 StartAtIndex, int32 Priority, float& OutAccumulatedWeight) -> bool
 				{
 					int32 NumAbsoluteBlends = 0;
+					int32 NumSkippedBlends = 0;
 					float AccumulatedWeight = 0.f;
 					for (int32 NextIndex = StartAtIndex; NextIndex < NumMixerEntries; ++NextIndex)
 					{
@@ -132,12 +133,21 @@ namespace UE::MovieScene
 						}
 						else if (NextMixerEntry.bRequiresBlend && !NextMixerEntry.bAdditive)
 						{
-							++NumAbsoluteBlends;
-							AccumulatedWeight += NextMixerEntry.PoseWeight;
+							if (!FMath::IsNearlyEqual(AccumulatedWeight, AccumulatedWeight + NextMixerEntry.PoseWeight, KINDA_SMALL_NUMBER))
+							{
+								++NumAbsoluteBlends;
+								AccumulatedWeight += NextMixerEntry.PoseWeight;
+							}
+							else
+							{
+								++NumSkippedBlends;
+							}
 						}
 					}
 
-					if (NumAbsoluteBlends > 1 && AccumulatedWeight != 0.0)
+					// When skipping blends, the accumulated weight still needs to be factored in for the absolute blends that aren't skipped,
+					// so if the total is greater than one, use a separate blend stack, and the accumulated weight.
+					if (NumAbsoluteBlends + NumSkippedBlends > 1 && !FMath::IsNearlyZero(AccumulatedWeight, KINDA_SMALL_NUMBER))
 					{
 						OutAccumulatedWeight = AccumulatedWeight;
 						return true;
@@ -169,6 +179,11 @@ namespace UE::MovieScene
 							continue;
 						}
 
+						if (bNeedsSeparateWeightStack && FMath::IsNearlyZero(MixerEntry.PoseWeight, KINDA_SMALL_NUMBER))
+						{
+							continue;
+						}
+
 						// Evaluate the pose itself
 						Mixer->EvaluationProgram->AppendTaskPtr(MixerEntry.EvalTask);
 
