Fix Crash when using the anim mixer with a source anim bp and control rig tracks.

ForEachEntity now checks if the anim instance returned from BindToSkeletalMeshComponent is not null, and will print a warning explaining that anim mixer Custom Anim Instance Target is incompatible with control rig and animation bluepritns.

diff --git a/Engine/Plugins/Experimental/MovieSceneAnimMixer/Source/MovieSceneAnimMixer/Private/MovieSceneAnimMixerModule.cpp b/Engine/Plugins/Experimental/MovieSceneAnimMixer/Source/MovieSceneAnimMixer/Private/MovieSceneAnimMixerModule.cpp
--- a/Engine/Plugins/Experimental/MovieSceneAnimMixer/Source/MovieSceneAnimMixer/Private/MovieSceneAnimMixerModule.cpp
+++ b/Engine/Plugins/Experimental/MovieSceneAnimMixer/Source/MovieSceneAnimMixer/Private/MovieSceneAnimMixerModule.cpp
@@ -4,6 +4,8 @@
 #include "Modules/ModuleInterface.h"
 #include "Modules/ModuleManager.h"
 
+DEFINE_LOG_CATEGORY(LogMovieSceneAnimMixer);
+
 namespace UE::MovieScene
 {
 
diff --git a/Engine/Plugins/Experimental/MovieSceneAnimMixer/Source/MovieSceneAnimMixer/Private/Systems/MovieSceneAnimInstanceTargetSystem.cpp b/Engine/Plugins/Experimental/MovieSceneAnimMixer/Source/MovieSceneAnimMixer/Private/Systems/MovieSceneAnimInstanceTargetSystem.cpp
--- a/Engine/Plugins/Experimental/MovieSceneAnimMixer/Source/MovieSceneAnimMixer/Private/Systems/MovieSceneAnimInstanceTargetSystem.cpp
+++ b/Engine/Plugins/Experimental/MovieSceneAnimMixer/Source/MovieSceneAnimMixer/Private/Systems/MovieSceneAnimInstanceTargetSystem.cpp
@@ -47,6 +47,8 @@
 #include "UnrealEdGlobals.h"
 #endif
 
+#include "MovieSceneAnimMixerModule.h"
+
 #include UE_INLINE_GENERATED_CPP_BY_NAME(MovieSceneAnimInstanceTargetSystem)
 
 
@@ -432,13 +434,14 @@ namespace UE::MovieScene
 		(
 			FMovieSceneEntityID EntityID,
 			FObjectComponent MeshComponent,
-			const TInstancedStruct<FMovieSceneMixedAnimationTarget>& Target,
+			TInstancedStruct<FMovieSceneMixedAnimationTarget>& Target,
 			const TSharedPtr<FAnimNextEvaluationTask>& MixerTask,
 			const TSharedPtr<FMovieSceneMixerRootMotionComponentData>* RootMotionData
 		) const
 		{
 			// Invalid or incompatible target, return.
-			if (!Target.GetPtr<FMovieSceneAnimInstanceTarget>())
+			FMovieSceneAnimInstanceTarget* AnimInstanceTarget = Target.GetMutablePtr<FMovieSceneAnimInstanceTarget>(); 
+			if (!AnimInstanceTarget)
 			{
 				return;
 			}
@@ -472,8 +475,23 @@ namespace UE::MovieScene
 			
 			USequencerMixedAnimInstance* SequencerInstance = FAnimCustomInstanceHelper::BindToSkeletalMeshComponent<USequencerMixedAnimInstance>(SkeletalMeshComponent, bWasCreated);
 
-			SequencerInstance->SetMixerTask(MixerTask);
-			SequencerInstance->SetRootMotion(RootMotionData ? *RootMotionData : TSharedPtr<FMovieSceneMixerRootMotionComponentData>());
+			if (SequencerInstance)
+			{
+				SequencerInstance->SetMixerTask(MixerTask);
+				SequencerInstance->SetRootMotion(RootMotionData ? *RootMotionData : TSharedPtr<FMovieSceneMixerRootMotionComponentData>());
+			}
+			else if (!AnimInstanceTarget->HasFiredWarningForTarget())
+			{
+				// Control rig also uses FAnimCustomInstanceHelper::BindToSkeletalMeshComponent to create an anim instance that it can use.
+				// If there is an existing animation blueprint specified on the skeletal mesh, that can cause this to fail when both the anim mixer and control rig try to bind to it.
+				UE_LOG(LogMovieSceneAnimMixer, Warning,
+				       TEXT(
+					       "Unable to bind anim mixer custom instance to skeletal mesh: %s on actor: %s. This is usually caused by a conflict with a control rig track. Currently it is unsupported to use the Sequenccer Anim Mixer with a Custom Anim Instance target and control rig tracks, if the actor has an animation blueprint assigned. Please remove one to resolve this issue."
+				       ), *SkeletalMeshComponent->GetName(), *SkeletalMeshComponent->GetOwner()->GetName())
+				
+				// Don't print to the log every frame.
+				AnimInstanceTarget->SetHasFiredWarningForTarget(true);
+			}
 
 #if WITH_EDITOR
 			if (GEditor && bWasCreated)
@@ -568,7 +586,7 @@ void UMovieSceneAnimInstanceTargetSystem::OnSchedulePersistentTasks(UE::MovieSce
 	FTaskID EvaluateTask = FEntityTaskBuilder()
 		.ReadEntityIDs()
 		.Read(AnimMixerComponents->MeshComponent)
-		.Read(AnimMixerComponents->Target)
+		.Write(AnimMixerComponents->Target)
 		.Read(AnimMixerComponents->MixerTask)
 		.ReadOptional(AnimMixerComponents->MixerRootMotion)
 		.SetParams(Params)
diff --git a/Engine/Plugins/Experimental/MovieSceneAnimMixer/Source/MovieSceneAnimMixer/Public/MovieSceneAnimMixerModule.h b/Engine/Plugins/Experimental/MovieSceneAnimMixer/Source/MovieSceneAnimMixer/Public/MovieSceneAnimMixerModule.h
--- a/Engine/Plugins/Experimental/MovieSceneAnimMixer/Source/MovieSceneAnimMixer/Public/MovieSceneAnimMixerModule.h
+++ b/Engine/Plugins/Experimental/MovieSceneAnimMixer/Source/MovieSceneAnimMixer/Public/MovieSceneAnimMixerModule.h
@@ -5,6 +5,8 @@
 #include "CoreMinimal.h"
 #include "Modules/ModuleInterface.h"
 
+MOVIESCENEANIMMIXER_API DECLARE_LOG_CATEGORY_EXTERN(LogMovieSceneAnimMixer, Log, All);
+
 struct FAnimationUpdateContext;
 
 namespace UE::AnimNext
diff --git a/Engine/Source/Runtime/MovieSceneTracks/Public/MovieSceneMixedAnimationTarget.h b/Engine/Source/Runtime/MovieSceneTracks/Public/MovieSceneMixedAnimationTarget.h
--- a/Engine/Source/Runtime/MovieSceneTracks/Public/MovieSceneMixedAnimationTarget.h
+++ b/Engine/Source/Runtime/MovieSceneTracks/Public/MovieSceneMixedAnimationTarget.h
@@ -23,4 +23,11 @@ struct MOVIESCENETRACKS_API FMovieSceneMixedAnimationTarget
 	{
 		return GetTypeHash(FMovieSceneMixedAnimationTarget::StaticStruct());
 	}
+
+	bool HasFiredWarningForTarget() const { return bHasFiredWarningForTarget; }
+
+	void SetHasFiredWarningForTarget(bool bNewHasFiredState) { bHasFiredWarningForTarget = bNewHasFiredState; }
+
+private:
+	bool bHasFiredWarningForTarget = false;
 };
\ No newline at end of file
