Sequencer: Ensure that skeletal mesh restore state, and VisibilityBasedTickOption is set up correctly for skeletal meshes when using the AnimMixer plugin

diff --git a/Engine/Plugins/Animation/ControlRig/Source/ControlRig/Private/Sequencer/MovieSceneControlRigParameterTemplate.cpp b/Engine/Plugins/Animation/ControlRig/Source/ControlRig/Private/Sequencer/MovieSceneControlRigParameterTemplate.cpp
--- a/Engine/Plugins/Animation/ControlRig/Source/ControlRig/Private/Sequencer/MovieSceneControlRigParameterTemplate.cpp
+++ b/Engine/Plugins/Animation/ControlRig/Source/ControlRig/Private/Sequencer/MovieSceneControlRigParameterTemplate.cpp
@@ -926,7 +926,7 @@ struct FControlRigParameterPreAnimatedTokenProducer : IMovieScenePreAnimatedToke
 									SkeletalMeshComponent->MarkRenderTransformDirty();
 									SkeletalMeshComponent->MarkRenderDynamicDataDirty();
 								}
-								SkeletalMeshRestoreState.RestoreState(SkeletalMeshComponent);
+								SkeletalMeshRestoreState.RestoreState();
 
 								if (SkeletalMeshComponent->GetAnimationMode() != AnimationMode)
 								{
diff --git a/Engine/Plugins/Animation/ControlRig/Source/ControlRig/Private/Sequencer/MovieSceneControlRigSystem.cpp b/Engine/Plugins/Animation/ControlRig/Source/ControlRig/Private/Sequencer/MovieSceneControlRigSystem.cpp
--- a/Engine/Plugins/Animation/ControlRig/Source/ControlRig/Private/Sequencer/MovieSceneControlRigSystem.cpp
+++ b/Engine/Plugins/Animation/ControlRig/Source/ControlRig/Private/Sequencer/MovieSceneControlRigSystem.cpp
@@ -809,7 +809,7 @@ void FPreAnimatedControlRigTraits::RestorePreAnimatedValue(const FObjectKey& InO
 					SkeletalMeshComponent->MarkRenderTransformDirty();
 					SkeletalMeshComponent->MarkRenderDynamicDataDirty();
 				}
-				State.SkeletalMeshRestoreState.RestoreState(SkeletalMeshComponent);
+				State.SkeletalMeshRestoreState.RestoreState();
 
 				if (SkeletalMeshComponent->GetAnimationMode() != State.AnimationMode)
 				{
diff --git a/Engine/Plugins/Experimental/MovieSceneAnimMixer/Source/MovieSceneAnimMixer/Private/Systems/MovieSceneAnimBlueprintTargetSystem.cpp b/Engine/Plugins/Experimental/MovieSceneAnimMixer/Source/MovieSceneAnimMixer/Private/Systems/MovieSceneAnimBlueprintTargetSystem.cpp
--- a/Engine/Plugins/Experimental/MovieSceneAnimMixer/Source/MovieSceneAnimMixer/Private/Systems/MovieSceneAnimBlueprintTargetSystem.cpp
+++ b/Engine/Plugins/Experimental/MovieSceneAnimMixer/Source/MovieSceneAnimMixer/Private/Systems/MovieSceneAnimBlueprintTargetSystem.cpp
@@ -62,7 +62,7 @@ namespace UE::MovieScene
 				MixerSubsystem->ResetEvalTasks();
 			}
 
-			InOutCachedValue.MeshRestoreState.RestoreLOD(Component);
+			InOutCachedValue.MeshRestoreState.RestoreLOD();
 
 			// Restore pose after unbinding to force the restored pose
 			Component->SetUpdateAnimationInEditor(true);
@@ -78,7 +78,7 @@ namespace UE::MovieScene
 				Component->MarkRenderDynamicDataDirty();
 			}
 
-			InOutCachedValue.MeshRestoreState.RestoreState(Component);
+			InOutCachedValue.MeshRestoreState.RestoreState();
 
 			// if not game world, don't clean this up
 			if (Component->GetWorld() != nullptr && Component->GetWorld()->IsGameWorld() == false)
diff --git a/Engine/Plugins/Experimental/MovieSceneAnimMixer/Source/MovieSceneAnimMixer/Private/Systems/MovieSceneAnimInstanceTargetSystem.cpp b/Engine/Plugins/Experimental/MovieSceneAnimMixer/Source/MovieSceneAnimMixer/Private/Systems/MovieSceneAnimInstanceTargetSystem.cpp
--- a/Engine/Plugins/Experimental/MovieSceneAnimMixer/Source/MovieSceneAnimMixer/Private/Systems/MovieSceneAnimInstanceTargetSystem.cpp
+++ b/Engine/Plugins/Experimental/MovieSceneAnimMixer/Source/MovieSceneAnimMixer/Private/Systems/MovieSceneAnimInstanceTargetSystem.cpp
@@ -355,7 +355,7 @@ namespace UE::MovieScene
 			FAnimCustomInstanceHelper::UnbindFromSkeletalMeshComponent<USequencerMixedAnimInstance>(Component);
 
 			// Restore LOD before reinitializing anim instance
-			InOutCachedValue.SkeletalMeshRestoreState.RestoreLOD(Component);
+			InOutCachedValue.SkeletalMeshRestoreState.RestoreLOD();
 
 			if (Component->GetAnimationMode() != InOutCachedValue.AnimationMode)
 			{
@@ -392,7 +392,7 @@ namespace UE::MovieScene
 			}
 
 			// Reset the mesh component update flag and animation mode to what they were before we animated the object
-			InOutCachedValue.SkeletalMeshRestoreState.RestoreState(Component);
+			InOutCachedValue.SkeletalMeshRestoreState.RestoreState();
 
 			// if not game world, don't clean this up
 			if (Component->GetWorld() != nullptr && Component->GetWorld()->IsGameWorld() == false)
diff --git a/Engine/Plugins/Experimental/MovieSceneAnimMixer/Source/MovieSceneAnimMixer/Private/Systems/MovieSceneAnimNextTargetSystem.cpp b/Engine/Plugins/Experimental/MovieSceneAnimMixer/Source/MovieSceneAnimMixer/Private/Systems/MovieSceneAnimNextTargetSystem.cpp
--- a/Engine/Plugins/Experimental/MovieSceneAnimMixer/Source/MovieSceneAnimMixer/Private/Systems/MovieSceneAnimNextTargetSystem.cpp
+++ b/Engine/Plugins/Experimental/MovieSceneAnimMixer/Source/MovieSceneAnimMixer/Private/Systems/MovieSceneAnimNextTargetSystem.cpp
@@ -9,6 +9,7 @@
 #include "EntitySystem/MovieSceneEntitySystemRunner.h"
 #include "EntitySystem/MovieSceneEntitySystemTask.h"
 #include "MovieSceneTracksComponentTypes.h"
+#include "SkeletalMeshRestoreState.h"
 
 #include "Evaluation/PreAnimatedState/IMovieScenePreAnimatedStorage.h"
 #include "Evaluation/PreAnimatedState/MovieScenePreAnimatedObjectStorage.h"
@@ -63,16 +64,10 @@ namespace UE::MovieScene
 	/** Pre-animation traits for skeletal animations */
 	struct FPreAnimatedAnimNextMixedSkeletalAnimationTraits : FBoundObjectPreAnimatedStateTraits
 	{
-		using KeyType = FObjectKey;
+		using KeyType = TTuple<FObjectKey, FName>;
 		using StorageType = FPreAnimatedAnimNextState;
 
-		StorageType CachePreAnimatedValue(FObjectKey InKey)
-		{
-			check(false);
-			return FPreAnimatedAnimNextState();
-		}
-
-		void RestorePreAnimatedValue(const KeyType& Object, StorageType& InOutCachedValue, const FRestoreStateParams& Params)
+		void RestorePreAnimatedValue(const KeyType& ObjectAndInjectionSite, StorageType& InOutCachedValue, const FRestoreStateParams& Params)
 		{
 			// Uninject
 			if (InOutCachedValue.InjectionHandle.IsValid())
@@ -87,21 +82,92 @@ namespace UE::MovieScene
 			// Reset the anim next module back to not being enabled if applicable
 			if (!InOutCachedValue.bAnimNextModuleEnabled)
 			{
-				if (UAnimNextComponent* AnimNextComponent = Cast<UAnimNextComponent>(Object.ResolveObjectPtr()))
+				if (UAnimNextComponent* AnimNextComponent = Cast<UAnimNextComponent>(ObjectAndInjectionSite.Key.ResolveObjectPtr()))
 				{
 					AnimNextComponent->SetEnabled(false);
 				}
 			}
 		}
 	};
+	struct FPreAnimatedSkelMeshComponentTraits : FBoundObjectPreAnimatedStateTraits
+	{
+		using KeyType = FObjectKey;
+		using StorageType = FSkeletalMeshEditorParams;
+
+		FSkeletalMeshEditorParams CachePreAnimatedValue(const KeyType& Object)
+		{
+			USkeletalMeshComponent* SkelMeshComp = CastChecked<USkeletalMeshComponent>(Object.ResolveObjectPtr());
+
+			FSkeletalMeshEditorParams State;
+			State.SaveState(SkelMeshComp);
+			return State;
+		}
+
+		void RestorePreAnimatedValue(const KeyType& Object, FSkeletalMeshEditorParams& InOutCachedValue, const FRestoreStateParams& Params)
+		{
+			if (Object.ResolveObjectPtr() != nullptr)
+			{
+				InOutCachedValue.RestoreState();
+			}
+		}
+	};
 
 	/** Pre-animation storage for anim next target */
-	struct FPreAnimatedAnimNextMixedSkeletalAnimationStorage : TPreAnimatedStateStorage_ObjectTraits<FPreAnimatedAnimNextMixedSkeletalAnimationTraits>
+	struct FPreAnimatedAnimNextMixedSkeletalAnimationStorage : TPreAnimatedStateStorage<FPreAnimatedAnimNextMixedSkeletalAnimationTraits>
 	{
 		static TAutoRegisterPreAnimatedStorageID<FPreAnimatedAnimNextMixedSkeletalAnimationStorage> StorageID;
+
+		void BeginTrackingEntity(FMovieSceneEntityID EntityID, bool bWantsRestoreState, FRootInstanceHandle RootInstanceHandle, UAnimNextComponent* Component, FName InjectionSite)
+		{
+			if (!this->ParentExtension->IsCapturingGlobalState() && !bWantsRestoreState)
+			{
+				return;
+			}
+
+			FPreAnimatedEntityCaptureSource* EntityMetaData = this->ParentExtension->GetOrCreateEntityMetaData();
+
+			KeyType Key{ Component, InjectionSite };
+	
+			FPreAnimatedStorageIndex       StorageIndex = this->GetOrCreateStorageIndex(Key);
+			FPreAnimatedStorageGroupHandle GroupHandle  = this->Traits.MakeGroup(Component);
+			FPreAnimatedStateEntry Entry { GroupHandle, FPreAnimatedStateCachedValueHandle{ this->StorageID, StorageIndex } };
+
+			EntityMetaData->BeginTrackingEntity(Entry, EntityID, RootInstanceHandle, bWantsRestoreState);
+		}
+
+		template<typename OnCacheValue /* StorageType(const KeyType&) */>
+		void CachePreAnimatedValue(UAnimNextComponent* Component, FName InjectionSite, OnCacheValue&& CacheCallback)
+		{
+			EPreAnimatedCaptureSourceTracking TrackingMode = EPreAnimatedCaptureSourceTracking::CacheIfTracked;
+
+			if (this->ShouldTrackCaptureSource(TrackingMode, Component, InjectionSite))
+			{
+				KeyType Key{ Component, InjectionSite };
+	
+				FPreAnimatedStorageIndex       StorageIndex = this->GetOrCreateStorageIndex(Key);
+				FPreAnimatedStorageGroupHandle GroupHandle  = this->Traits.MakeGroup(Component);
+				FPreAnimatedStateEntry Entry { GroupHandle, FPreAnimatedStateCachedValueHandle{ this->StorageID, StorageIndex } };
+
+				this->TrackCaptureSource(Entry, TrackingMode);
+
+				EPreAnimatedStorageRequirement StorageRequirement = this->ParentExtension->GetStorageRequirement(Entry);
+				if (!this->IsStorageRequirementSatisfied(StorageIndex, StorageRequirement))
+				{
+					StorageType NewValue = CacheCallback(Component);
+					this->AssignPreAnimatedValue(StorageIndex, StorageRequirement, MoveTemp(NewValue));
+				}
+			}
+		}
+
+	};
+
+	struct FPreAnimatedSkelMeshComponentStorage : TPreAnimatedStateStorage_ObjectTraits<FPreAnimatedSkelMeshComponentTraits>
+	{
+		static TAutoRegisterPreAnimatedStorageID<FPreAnimatedSkelMeshComponentStorage> StorageID;
 	};
 
 	TAutoRegisterPreAnimatedStorageID<FPreAnimatedAnimNextMixedSkeletalAnimationStorage> FPreAnimatedAnimNextMixedSkeletalAnimationStorage::StorageID;
+	TAutoRegisterPreAnimatedStorageID<FPreAnimatedSkelMeshComponentStorage> FPreAnimatedSkelMeshComponentStorage::StorageID;
 
 	/** Task for evaluating and applying animation tasks */
 	struct FEvaluateAnimNextTasks
@@ -112,6 +178,7 @@ namespace UE::MovieScene
 		UMovieSceneAnimNextTargetSystem* System;
 
 		TSharedPtr<FPreAnimatedAnimNextMixedSkeletalAnimationStorage> PreAnimatedStorage;
+		TSharedPtr<FPreAnimatedSkelMeshComponentStorage> PreAnimatedSkelMeshCompStorage;
 
 	public:
 
@@ -120,6 +187,7 @@ namespace UE::MovieScene
 			, System(InSystem)
 		{
 			PreAnimatedStorage = InLinker->PreAnimatedState.GetOrCreateStorage<FPreAnimatedAnimNextMixedSkeletalAnimationStorage>();
+			PreAnimatedSkelMeshCompStorage = InLinker->PreAnimatedState.GetOrCreateStorage<FPreAnimatedSkelMeshComponentStorage>();
 		}
 
 		void ForEachEntity
@@ -140,15 +208,15 @@ namespace UE::MovieScene
 			// TODO: It's valid (or will be) to use Anim Next without a Skeletal Mesh Component.
 			// We'll need to refactor this to be able to pass an Anim Next component here and refactor core anim tracks to match.
 
-			UAnimNextComponent* AnimNextComponent = nullptr;
-
-			AnimNextComponent = Cast<UAnimNextComponent>(ObjectComponent.GetObject());
+			UObject* BoundMesh = ObjectComponent.GetObject();
+			USkeletalMeshComponent* SkeletalMeshComponent = Cast<USkeletalMeshComponent>(BoundMesh);
+			UAnimNextComponent* AnimNextComponent = Cast<UAnimNextComponent>(BoundMesh);
 
 			if (!AnimNextComponent)
 			{
 				AActor* Actor = nullptr;
 				// Could be a skeletal mesh component
-				if (UActorComponent* ActorComponent = Cast<UActorComponent>(ObjectComponent.GetObject()))
+				if (UActorComponent* ActorComponent = Cast<UActorComponent>(BoundMesh))
 				{
 					Actor = ActorComponent->GetOwner();
 				}
@@ -186,15 +254,22 @@ namespace UE::MovieScene
 					TargetData->InjectionRequestHandle = UE::AnimNext::FInjectionUtils::InjectEvaluationModifier(AnimNextComponent, TargetData->Modifier.ToSharedRef(), TargetData->InjectionSite);
 
 					// For now, we always restore state
-					PreAnimatedStorage->BeginTrackingEntity(EntityID, true, FRootInstanceHandle(), AnimNextComponent);
+					PreAnimatedStorage->BeginTrackingEntity(EntityID, true, FRootInstanceHandle(), AnimNextComponent, InjectionSite.DesiredSiteName);
+					PreAnimatedSkelMeshCompStorage->BeginTrackingEntity(EntityID, true, FRootInstanceHandle(), SkeletalMeshComponent);
 
 					bool bAnimNextModuleEnabled = AnimNextComponent->IsEnabled();
 
-					FCachePreAnimatedValueParams CacheParams;
-					PreAnimatedStorage->CachePreAnimatedValue(
-						CacheParams,
-						AnimNextComponent,
-						[InjectionHandle = TargetData->InjectionRequestHandle, AnimSystem=System, bAnimNextModuleEnabled](FObjectKey InKey) { return FPreAnimatedAnimNextState{InjectionHandle, AnimSystem, bAnimNextModuleEnabled }; });
+					auto OnCacheAnimatedState = [this, TargetData, bAnimNextModuleEnabled](UAnimNextComponent*)
+					{
+						FPreAnimatedAnimNextState State;
+						State.InjectionHandle = TargetData->InjectionRequestHandle;
+						State.System = this->System;
+						State.bAnimNextModuleEnabled = bAnimNextModuleEnabled;;
+						return State;
+					};
+
+					PreAnimatedStorage->CachePreAnimatedValue(AnimNextComponent, InjectionSite.DesiredSiteName, OnCacheAnimatedState);
+					PreAnimatedSkelMeshCompStorage->CachePreAnimatedValue(FCachePreAnimatedValueParams(), SkeletalMeshComponent);
 
 					// Force enable the anim next module if not currently enabled
 					if (!bAnimNextModuleEnabled)
diff --git a/Engine/Source/Runtime/MovieScene/Public/SkeletalMeshRestoreState.h b/Engine/Source/Runtime/MovieScene/Public/SkeletalMeshRestoreState.h
--- a/Engine/Source/Runtime/MovieScene/Public/SkeletalMeshRestoreState.h
+++ b/Engine/Source/Runtime/MovieScene/Public/SkeletalMeshRestoreState.h
@@ -77,14 +77,25 @@ struct FSkeletalMeshRestoreState
 		}
 	}
 
+	UE_DEPRECATED(5.6, TEXT("InComponent is not required"))
 	void RestoreState(USkeletalMeshComponent* InComponent)
+	{
+		RestoreState();
+	}
+	UE_DEPRECATED(5.6, TEXT("InComponent is not required"))
+	void RestoreLOD(USkeletalMeshComponent* InComponent)
+	{
+		RestoreLOD();
+	}
+
+	void RestoreState()
 	{
 		for (FSkeletalMeshEditorParams& ChildParams : SkeletalMeshCompEditorParams)
 		{
 			ChildParams.RestoreState();
 		}
 	}
-	void RestoreLOD(USkeletalMeshComponent* InComponent)
+	void RestoreLOD()
 	{
 		for (FSkeletalMeshEditorParams& ChildParams : SkeletalMeshCompEditorParams)
 		{
diff --git a/Engine/Source/Runtime/MovieSceneTracks/Private/Systems/MovieSceneSkeletalAnimationSystem.cpp b/Engine/Source/Runtime/MovieSceneTracks/Private/Systems/MovieSceneSkeletalAnimationSystem.cpp
--- a/Engine/Source/Runtime/MovieSceneTracks/Private/Systems/MovieSceneSkeletalAnimationSystem.cpp
+++ b/Engine/Source/Runtime/MovieSceneTracks/Private/Systems/MovieSceneSkeletalAnimationSystem.cpp
@@ -130,7 +130,7 @@ struct FPreAnimatedSkeletalAnimationTraits : FBoundObjectPreAnimatedStateTraits
 		FAnimCustomInstanceHelper::UnbindFromSkeletalMeshComponent<UAnimSequencerInstance>(Component);
 
 		// Restore LOD before reinitializing anim instance
-		InOutCachedValue.SkeletalMeshRestoreState.RestoreLOD(Component);
+		InOutCachedValue.SkeletalMeshRestoreState.RestoreLOD();
 
 		if (Component->GetAnimationMode() != InOutCachedValue.AnimationMode)
 		{
@@ -168,7 +168,7 @@ struct FPreAnimatedSkeletalAnimationTraits : FBoundObjectPreAnimatedStateTraits
 		}
 
 		// Reset the mesh component update flag and animation mode to what they were before we animated the object
-		InOutCachedValue.SkeletalMeshRestoreState.RestoreState(Component);
+		InOutCachedValue.SkeletalMeshRestoreState.RestoreState();
 
 		// if not game world, don't clean this up
 		if (Component->GetWorld() != nullptr && Component->GetWorld()->IsGameWorld() == false)
