Do not skip redundant transform updates for skinned mesh components when controlled by both AnimNext and Sequencer

diff --git a/Engine/Plugins/Experimental/MovieSceneAnimMixer/Source/MovieSceneAnimMixer/Private/Systems/MovieSceneAnimNextTargetSystem.cpp b/Engine/Plugins/Experimental/MovieSceneAnimMixer/Source/MovieSceneAnimMixer/Private/Systems/MovieSceneAnimNextTargetSystem.cpp
--- a/Engine/Plugins/Experimental/MovieSceneAnimMixer/Source/MovieSceneAnimMixer/Private/Systems/MovieSceneAnimNextTargetSystem.cpp
+++ b/Engine/Plugins/Experimental/MovieSceneAnimMixer/Source/MovieSceneAnimMixer/Private/Systems/MovieSceneAnimNextTargetSystem.cpp
@@ -36,6 +36,8 @@
 #include "Injection/InjectionUtils.h"
 #include "TraitInterfaces/IEvaluate.h"
 
+#include "PrimitiveSceneProxy.h"
+
 
 #include UE_INLINE_GENERATED_CPP_BY_NAME(MovieSceneAnimNextTargetSystem)
 
@@ -100,13 +102,16 @@ namespace UE::MovieScene
 
 			FSkeletalMeshEditorParams State;
 			State.SaveState(SkelMeshComp);
+			SkelMeshComp->GetSceneProxy()->SetCanSkipRedundantTransformUpdates(false);
 			return State;
 		}
 
 		void RestorePreAnimatedValue(const KeyType& Object, FSkeletalMeshEditorParams& InOutCachedValue, const FRestoreStateParams& Params)
 		{
-			if (Object.ResolveObjectPtr() != nullptr)
+			USkeletalMeshComponent* SkelMeshComp = Cast<USkeletalMeshComponent>(Object.ResolveObjectPtr());
+			if (SkelMeshComp != nullptr)
 			{
+				SkelMeshComp->GetSceneProxy()->SetCanSkipRedundantTransformUpdates(true);
 				InOutCachedValue.RestoreState();
 			}
 		}
@@ -264,7 +269,7 @@ namespace UE::MovieScene
 						FPreAnimatedAnimNextState State;
 						State.InjectionHandle = TargetData->InjectionRequestHandle;
 						State.System = this->System;
-						State.bAnimNextModuleEnabled = bAnimNextModuleEnabled;;
+						State.bAnimNextModuleEnabled = bAnimNextModuleEnabled;
 						return State;
 					};
 
diff --git a/Engine/Source/Runtime/Engine/Classes/Components/SkinnedMeshComponent.h b/Engine/Source/Runtime/Engine/Classes/Components/SkinnedMeshComponent.h
--- a/Engine/Source/Runtime/Engine/Classes/Components/SkinnedMeshComponent.h
+++ b/Engine/Source/Runtime/Engine/Classes/Components/SkinnedMeshComponent.h
@@ -1596,7 +1596,11 @@ public:
 	 * Checks/updates material usage on proxy based on current morph target usage
 	 */
 	ENGINE_API void UpdateMorphMaterialUsageOnProxy();
-	
+
+	/**
+	 * Toggles whether this skinned mesh component can issue duplicate transform updates
+	 */
+	ENGINE_API void AllowDuplicateTransformUpdates(bool bAllow);
 
 	/** Access ComponentSpaceTransforms for reading */
 	const TArray<FTransform>& GetComponentSpaceTransforms() const 
diff --git a/Engine/Source/Runtime/Engine/Private/Components/SkinnedMeshComponent.cpp b/Engine/Source/Runtime/Engine/Private/Components/SkinnedMeshComponent.cpp
--- a/Engine/Source/Runtime/Engine/Private/Components/SkinnedMeshComponent.cpp
+++ b/Engine/Source/Runtime/Engine/Private/Components/SkinnedMeshComponent.cpp
@@ -609,6 +609,14 @@ FPrimitiveSceneProxy* USkinnedMeshComponent::CreateSceneProxy()
 	return Result;
 }
 
+void USkinnedMeshComponent::AllowDuplicateTransformUpdates(bool bAllow)
+{
+	if (SceneProxy)
+	{
+		SceneProxy->SetCanSkipRedundantTransformUpdates(!bAllow);
+	}
+}
+
 // UObject interface
 // Override to have counting working better
 void USkinnedMeshComponent::Serialize(FArchive& Ar)
diff --git a/Engine/Source/Runtime/Engine/Public/PrimitiveSceneProxy.h b/Engine/Source/Runtime/Engine/Public/PrimitiveSceneProxy.h
--- a/Engine/Source/Runtime/Engine/Public/PrimitiveSceneProxy.h
+++ b/Engine/Source/Runtime/Engine/Public/PrimitiveSceneProxy.h
@@ -574,6 +574,14 @@ public:
 		return bCanSkipRedundantTransformUpdates;
 	}
 
+	/**
+	* Set whether the proxy can skip redundant transform updates where applicable.
+	*/
+	void SetCanSkipRedundantTransformUpdates(bool bCanSkip)
+	{
+		bCanSkipRedundantTransformUpdates = bCanSkip;
+	}
+
 	/**
 	* @return true if the proxy uses distance cull fade.
 	*/
