Fix Crash From Animating an object that is being GCd

diff --git a/Engine/Plugins/Experimental/MovieSceneAnimMixer/Source/MovieSceneAnimMixer/Private/Systems/MovieSceneAnimBlueprintTargetSystem.cpp b/Engine/Plugins/Experimental/MovieSceneAnimMixer/Source/MovieSceneAnimMixer/Private/Systems/MovieSceneAnimBlueprintTargetSystem.cpp
--- a/Engine/Plugins/Experimental/MovieSceneAnimMixer/Source/MovieSceneAnimMixer/Private/Systems/MovieSceneAnimBlueprintTargetSystem.cpp
+++ b/Engine/Plugins/Experimental/MovieSceneAnimMixer/Source/MovieSceneAnimMixer/Private/Systems/MovieSceneAnimBlueprintTargetSystem.cpp
@@ -128,6 +128,12 @@ namespace UE::MovieScene
 				return;
 			}
 
+			// Mesh may no longer be valid
+			if (!IsValid(MeshComponent.GetObject()))
+			{
+				return;
+			}
+
 			const FName BlueprintNodeName = Target.GetPtr<FMovieSceneAnimBlueprintTarget>()->BlueprintNodeName;
 
 			USkeletalMeshComponent* SkeletalMeshComponent = CastChecked<USkeletalMeshComponent>(MeshComponent.GetObject());
