Sequencer: Initial Animation Mixer Track prototype

- Added a new track that has root motion options, skeletal animation root offset channels, and can contain stitch tracks

diff --git a/Engine/Plugins/Experimental/Mover/Source/Mover/Private/DefaultMovementSet/LayeredMoves/RootMotionAttributeLayeredMove.cpp b/Engine/Plugins/Experimental/Mover/Source/Mover/Private/DefaultMovementSet/LayeredMoves/RootMotionAttributeLayeredMove.cpp
--- a/Engine/Plugins/Experimental/Mover/Source/Mover/Private/DefaultMovementSet/LayeredMoves/RootMotionAttributeLayeredMove.cpp
+++ b/Engine/Plugins/Experimental/Mover/Source/Mover/Private/DefaultMovementSet/LayeredMoves/RootMotionAttributeLayeredMove.cpp
@@ -32,6 +32,9 @@ static const FName RootMotionAttributeName = "RootMotionDelta";
 static const UE::Anim::FAttributeId RootMotionAttributeId = { RootMotionAttributeName, FCompactPoseBoneIndex(0) };
 
 
+static const FName RootTransformAttributeName = "RootTransform";
+static const UE::Anim::FAttributeId RootTransformAttributeId = { RootTransformAttributeName, FCompactPoseBoneIndex(0) };
+
 FLayeredMove_RootMotionAttribute::FLayeredMove_RootMotionAttribute()
 {
 	DurationMs = -1.f;
@@ -49,6 +52,8 @@ bool FLayeredMove_RootMotionAttribute::GenerateMove(const FMoverTickStartData& S
 
 	const float DeltaSeconds = TimeStep.StepMs / 1000.f;
 
+	bool bDidAttrHaveRootTransform = false;
+	FTransform RootTransform;
 	bool bDidAttrHaveRootMotion = false;
 	FTransform LocalRootMotion;
 	bool bHasValidWarpingContext = false;
@@ -67,8 +72,12 @@ bool FLayeredMove_RootMotionAttribute::GenerateMove(const FMoverTickStartData& S
 		if (Mesh)
 		{
 			// TODO: support options for different interpretations, such as velocity
-
-			if (const FTransformAnimationAttribute* RootMotionAttribute = Mesh->GetCustomAttributes().Find<FTransformAnimationAttribute>(RootMotionAttributeId))
+			if (const FTransformAnimationAttribute* RootTransformAttribute = Mesh->GetCustomAttributes().Find<FTransformAnimationAttribute>(RootTransformAttributeId))
+			{
+				RootTransform = RootTransformAttribute->Value;
+				bDidAttrHaveRootTransform = true;
+			}
+			else if (const FTransformAnimationAttribute* RootMotionAttribute = Mesh->GetCustomAttributes().Find<FTransformAnimationAttribute>(RootMotionAttributeId))
 			{
 				// NOTE this will only work for ticking modes that tick in time with the world tick, because it relies on the Mesh ticking at the same rate as the movement simulation. 
 				// For fixed-tick modes, the attribute would be better as an accumulator for the movement sim to consume, along with a time accumulation attribute.
@@ -121,8 +130,22 @@ bool FLayeredMove_RootMotionAttribute::GenerateMove(const FMoverTickStartData& S
 		WarpingContext			= WarpingContextForResim;
 	}
 
-	
-	if (bDidAttrHaveRootMotion)
+	if (bDidAttrHaveRootTransform)
+	{
+		const FMoverDefaultSyncState* SyncState = SimState.SyncState.SyncStateCollection.FindDataByType<FMoverDefaultSyncState>();
+
+		const FTransform SimActorTransform = FTransform(SyncState->GetOrientation_WorldSpace().Quaternion(), SyncState->GetLocation_WorldSpace());
+
+		OutProposedMove = FProposedMove();
+		OutProposedMove.MixMode = EMoveMixMode::OverrideAll;
+
+		// Convert the transform into linear and angular velocities
+		OutProposedMove.LinearVelocity  = (RootTransform.GetTranslation() - SimActorTransform.GetTranslation()) / DeltaSeconds;
+		OutProposedMove.AngularVelocity = ((RootTransform.GetRotation()-SimActorTransform.GetRotation()) / DeltaSeconds).Rotator();
+
+		return true;
+	}
+	else if (bDidAttrHaveRootMotion)
 	{
 		const FMoverDefaultSyncState* SyncState = SimState.SyncState.SyncStateCollection.FindDataByType<FMoverDefaultSyncState>();
 		const FCharacterDefaultInputs* InputCmd = SimState.InputCmd.InputCollection.FindDataByType<FCharacterDefaultInputs>();
diff --git a/Engine/Plugins/Experimental/MovieSceneAnimMixer/Source/MovieSceneAnimMixer/Private/AnimMixerComponentTypes.cpp b/Engine/Plugins/Experimental/MovieSceneAnimMixer/Source/MovieSceneAnimMixer/Private/AnimMixerComponentTypes.cpp
--- a/Engine/Plugins/Experimental/MovieSceneAnimMixer/Source/MovieSceneAnimMixer/Private/AnimMixerComponentTypes.cpp
+++ b/Engine/Plugins/Experimental/MovieSceneAnimMixer/Source/MovieSceneAnimMixer/Private/AnimMixerComponentTypes.cpp
@@ -11,6 +11,9 @@ namespace UE::MovieScene
 {
 
 	const UE::Anim::FAttributeId FAnimMixerComponentTypes::RootTransformAttributeId(TEXT("RootTransform"), FCompactPoseBoneIndex(0));
+	const UE::Anim::FAttributeId FAnimMixerComponentTypes::RootTransformBaseAttributeId(TEXT("RootTransformBase"), FCompactPoseBoneIndex(0));
+	const UE::Anim::FAttributeId FAnimMixerComponentTypes::RootTransformWeightAttributeId(TEXT("RootTransformWeight"), FCompactPoseBoneIndex(0));
+	const UE::Anim::FAttributeId FAnimMixerComponentTypes::LegacySwapRootBoneAttributeId(TEXT("LegacySwapRootBone"), FCompactPoseBoneIndex(0));
 
 	static TUniquePtr<FAnimMixerComponentTypes> GAnimMixerComponentTypes;
 	static bool GAnimMixerComponentTypesDestroyed = false;
@@ -43,6 +46,8 @@ namespace UE::MovieScene
 		ComponentRegistry->NewComponentType(&RootMotionSettings, TEXT("Root Motion Settings"));
 		ComponentRegistry->NewComponentType(&MeshComponent, TEXT("Mixed Animation Mesh Component"));
 		ComponentRegistry->NewComponentType(&MixerRootMotion, TEXT("Root Motion"));
+		ComponentRegistry->NewComponentType(&MixerEntry, TEXT("MixerEntry"));
+		ComponentRegistry->NewComponentType(&RootDestination, TEXT("Root Destination"));
 
 
 		Tags.RequiresBlending = ComponentRegistry->NewTag(TEXT("Requires Blending"));
@@ -54,6 +59,10 @@ namespace UE::MovieScene
 		ComponentRegistry->Factories.DuplicateChildComponent(RootMotionSettings);
 		ComponentRegistry->Factories.DuplicateChildComponent(MeshComponent);
 		ComponentRegistry->Factories.DuplicateChildComponent(MixerRootMotion);
+		ComponentRegistry->Factories.DuplicateChildComponent(RootDestination);
+
+		ComponentRegistry->Factories.DefineChildComponent(RootDestination, MixerRootMotion);
+		ComponentRegistry->Factories.DefineChildComponent(Task, MixerEntry);
 
 		ComponentRegistry->Factories.DefineChildComponent(Tags.RequiresBlending, Tags.RequiresBlending);
 
diff --git a/Engine/Plugins/Experimental/MovieSceneAnimMixer/Source/MovieSceneAnimMixer/Private/MovieSceneAnimationMixerTrack.cpp b/Engine/Plugins/Experimental/MovieSceneAnimMixer/Source/MovieSceneAnimMixer/Private/MovieSceneAnimationMixerTrack.cpp
new file mode 100644
--- /dev/null
+++ b/Engine/Plugins/Experimental/MovieSceneAnimMixer/Source/MovieSceneAnimMixer/Private/MovieSceneAnimationMixerTrack.cpp
@@ -0,0 +1,216 @@
+// Copyright Epic Games, Inc. All Rights Reserved.
+
+#include "MovieSceneAnimationMixerTrack.h"
+#include "MovieSceneSection.h"
+#include "Sections/MovieSceneSkeletalAnimationSection.h"
+#include "Channels/MovieSceneChannelProxy.h"
+#include "EntitySystem/BuiltInComponentTypes.h"
+#include "Misc/AxisDisplayInfo.h"
+
+UMovieSceneAnimationBaseTransformDecoration::UMovieSceneAnimationBaseTransformDecoration(const FObjectInitializer& ObjInit)
+	: Super(ObjInit)
+{
+	Location[0].SetDefault(0.0);
+	Location[1].SetDefault(0.0);
+	Location[2].SetDefault(0.0);
+	Rotation[0].SetDefault(0.0);
+	Rotation[1].SetDefault(0.0);
+	Rotation[2].SetDefault(0.0);
+}
+
+EMovieSceneChannelProxyType UMovieSceneAnimationBaseTransformDecoration::PopulateChannelProxy(FMovieSceneChannelProxyData& OutProxyData)
+{
+#if WITH_EDITOR
+
+	FText LocationGroup = NSLOCTEXT("BaseTransformDecoration", "RootBaseLocation", "Root Base Location");
+	{
+		FMovieSceneChannelMetaData ChannelMetaData;
+		ChannelMetaData.SetIdentifiers("RootBaseLocation.X", AxisDisplayInfo::GetAxisDisplayName(EAxisList::X), LocationGroup);
+		ChannelMetaData.Color = AxisDisplayInfo::GetAxisColor(EAxisList::X);
+		ChannelMetaData.WeakOwningObject = this;
+		ChannelMetaData.SortOrder = 0;
+		OutProxyData.Add(Location[0], ChannelMetaData, TMovieSceneExternalValue<double>());
+	}
+	{
+		FMovieSceneChannelMetaData ChannelMetaData;
+		ChannelMetaData.SetIdentifiers("RootBaseLocation.Y", AxisDisplayInfo::GetAxisDisplayName(EAxisList::Y), LocationGroup);
+		ChannelMetaData.Color = AxisDisplayInfo::GetAxisColor(EAxisList::Y);
+		ChannelMetaData.WeakOwningObject = this;
+		ChannelMetaData.SortOrder = 1;
+		OutProxyData.Add(Location[1], ChannelMetaData, TMovieSceneExternalValue<double>());
+	}
+	{
+		FMovieSceneChannelMetaData ChannelMetaData;
+		ChannelMetaData.SetIdentifiers("RootBaseLocation.Z", AxisDisplayInfo::GetAxisDisplayName(EAxisList::Z), LocationGroup);
+		ChannelMetaData.Color = AxisDisplayInfo::GetAxisColor(EAxisList::Z);
+		ChannelMetaData.WeakOwningObject = this;
+		ChannelMetaData.SortOrder = 2;
+		OutProxyData.Add(Location[2], ChannelMetaData, TMovieSceneExternalValue<double>());
+	}
+
+	FText RotationGroup = NSLOCTEXT("BaseTransformDecoration", "RootBaseRotation", "Root Base Rotation");
+	{
+		FMovieSceneChannelMetaData ChannelMetaData;
+		ChannelMetaData.SetIdentifiers("RootBaseRotation.X", NSLOCTEXT("BaseTransformDecoration", "RotationX", "Roll"), RotationGroup);
+		ChannelMetaData.Color = AxisDisplayInfo::GetAxisColor(EAxisList::X);
+		ChannelMetaData.WeakOwningObject = this;
+		ChannelMetaData.SortOrder = 3;
+		OutProxyData.Add(Rotation[0], ChannelMetaData, TMovieSceneExternalValue<double>());
+	}
+	{
+		FMovieSceneChannelMetaData ChannelMetaData;
+		ChannelMetaData.SetIdentifiers("RootBaseRotation.Y", NSLOCTEXT("BaseTransformDecoration", "RotationY", "Pitch"), RotationGroup);
+		ChannelMetaData.Color = AxisDisplayInfo::GetAxisColor(EAxisList::Y);
+		ChannelMetaData.WeakOwningObject = this;
+		ChannelMetaData.SortOrder = 4;
+		OutProxyData.Add(Rotation[1], ChannelMetaData, TMovieSceneExternalValue<double>());
+	}
+	{
+		FMovieSceneChannelMetaData ChannelMetaData;
+		ChannelMetaData.SetIdentifiers("RootBaseRotation.Z", NSLOCTEXT("BaseTransformDecoration", "RotationZ", "Yaw"), RotationGroup);
+		ChannelMetaData.Color = AxisDisplayInfo::GetAxisColor(EAxisList::Z);
+		ChannelMetaData.WeakOwningObject = this;
+		ChannelMetaData.SortOrder = 5;
+		OutProxyData.Add(Rotation[2], ChannelMetaData, TMovieSceneExternalValue<double>());
+	}
+
+#else
+
+	OutProxyData.Add(Location[0]);
+	OutProxyData.Add(Location[1]);
+	OutProxyData.Add(Location[2]);
+
+	OutProxyData.Add(Rotation[0]);
+	OutProxyData.Add(Rotation[1]);
+	OutProxyData.Add(Rotation[2]);
+
+#endif
+
+	return EMovieSceneChannelProxyType::Static;
+}
+
+void UMovieSceneAnimationBaseTransformDecoration::ExtendEntityImpl(UMovieSceneEntitySystemLinker* EntityLinker, const FEntityImportParams& Params, FImportedEntity* OutImportedEntity)
+{
+	using namespace UE::MovieScene;
+	FBuiltInComponentTypes* BuiltInComponents = FBuiltInComponentTypes::Get();
+
+	OutImportedEntity->AddBuilder(
+		FEntityBuilder()
+		.Add(BuiltInComponents->DoubleChannel[0], &Location[0])
+		.Add(BuiltInComponents->DoubleChannel[1], &Location[1])
+		.Add(BuiltInComponents->DoubleChannel[2], &Location[2])
+		.Add(BuiltInComponents->DoubleChannel[3], &Rotation[0])
+		.Add(BuiltInComponents->DoubleChannel[4], &Rotation[1])
+		.Add(BuiltInComponents->DoubleChannel[5], &Rotation[2])
+	);
+}
+
+
+UMovieSceneAnimationMixerTrack::UMovieSceneAnimationMixerTrack(const FObjectInitializer& ObjInit)
+	: Super(ObjInit)
+{
+	SupportedBlendTypes.Add(EMovieSceneBlendType::Absolute);
+
+#if WITH_EDITORONLY_DATA
+	TrackTint = FColor(66, 56, 88, 255);
+	bSupportsDefaultSections = false;
+#endif
+}
+
+
+bool UMovieSceneAnimationMixerTrack::SupportsType(TSubclassOf<UMovieSceneSection> SectionClass) const
+{
+	SectionClass.Get()->ImplementsInterface(IMovieSceneAnimationSectionInterface::UClassType::StaticClass());
+	return true;
+}
+
+bool UMovieSceneAnimationMixerTrack::FixRowIndices()
+{
+	// Group sections by their type on different rows
+	TSortedMap<UClass*, TArray<UMovieSceneSection*>> SectionsByType;
+	for (UMovieSceneSection* Section : AnimationSections)
+	{
+		SectionsByType.FindOrAdd(Section->GetClass()).Add(Section);
+	}
+
+	bool bMadeChanges = false;
+	int32 CurrentRowIndex = 0;
+
+	TArray<UClass*> SortedSectionTypes;
+	SectionsByType.GenerateKeyArray(SortedSectionTypes);
+
+	Algo::SortBy(SortedSectionTypes, [](UClass* Class){
+		IMovieSceneAnimationSectionInterface* Interface = Cast<IMovieSceneAnimationSectionInterface>(Class->GetDefaultObject());
+		return Interface ? Interface->GetRowSortOrder() : 0;
+	});
+
+	for (UClass* Key : SortedSectionTypes)
+	{
+		TArray<UMovieSceneSection*>& SectionsOfType = SectionsByType.FindChecked(Key);
+
+		// Sort the sections
+		Algo::SortBy(SectionsOfType, &UMovieSceneSection::GetRowIndex);
+
+		int32 PreviousIndex = SectionsOfType[0]->GetRowIndex();
+		for (UMovieSceneSection* Section : SectionsOfType)
+		{
+			const int32 ThisIndex = Section->GetRowIndex();
+			if (PreviousIndex != ThisIndex)
+			{
+				++CurrentRowIndex;
+				PreviousIndex = ThisIndex;
+			}
+
+			if (ThisIndex != CurrentRowIndex)
+			{
+				bMadeChanges = true;
+				Section->SetRowIndex(CurrentRowIndex);
+			}
+		}
+
+		++CurrentRowIndex;
+	}
+
+	return bMadeChanges;
+}
+
+void UMovieSceneAnimationMixerTrack::OnSectionAddedImpl(UMovieSceneSection* Section)
+{
+	if (UMovieSceneSkeletalAnimationSection* AnimSection = Cast<UMovieSceneSkeletalAnimationSection>(Section))
+	{
+		UMovieSceneAnimationSectionDecoration* NewDecoration = NewObject<UMovieSceneAnimationSectionDecoration>(AnimSection, UMovieSceneAnimationSectionDecoration::StaticClass(), NAME_None, RF_Transactional);
+		AnimSection->AddDecoration(NewDecoration);
+
+		UMovieSceneAnimationBaseTransformDecoration* BaseTransform = NewObject<UMovieSceneAnimationBaseTransformDecoration>(AnimSection, UMovieSceneAnimationBaseTransformDecoration::StaticClass(), NAME_None, RF_Transactional);
+		AnimSection->AddDecoration(BaseTransform);
+	}
+}
+
+
+#if WITH_EDITORONLY_DATA
+
+FText UMovieSceneAnimationMixerTrack::GetTrackRowDisplayName(int32 RowIndex) const
+{
+	// Row display name is defined by the first section in that row
+	for (UMovieSceneSection* Section : AnimationSections)
+	{
+		if (Section->GetRowIndex() == RowIndex)
+		{
+			return Section->GetClass()->GetDisplayNameText();
+		}
+	}
+
+	return Super::GetTrackRowDisplayName(RowIndex);
+}
+
+FText UMovieSceneAnimationMixerTrack::GetDefaultDisplayName() const
+{
+	return NSLOCTEXT("AnimMixer", "DefaultTrackName", "Animation Mixer");
+}
+
+bool UMovieSceneAnimationMixerTrack::CanRename() const
+{
+	return true;
+}
+
+#endif // WITH_EDITORONLY_DATA
diff --git a/Engine/Plugins/Experimental/MovieSceneAnimMixer/Source/MovieSceneAnimMixer/Private/Systems/MovieSceneAnimInstanceTargetSystem.cpp b/Engine/Plugins/Experimental/MovieSceneAnimMixer/Source/MovieSceneAnimMixer/Private/Systems/MovieSceneAnimInstanceTargetSystem.cpp
--- a/Engine/Plugins/Experimental/MovieSceneAnimMixer/Source/MovieSceneAnimMixer/Private/Systems/MovieSceneAnimInstanceTargetSystem.cpp
+++ b/Engine/Plugins/Experimental/MovieSceneAnimMixer/Source/MovieSceneAnimMixer/Private/Systems/MovieSceneAnimInstanceTargetSystem.cpp
@@ -98,11 +98,6 @@ void USequencerMixedAnimInstance::SetMixerTask(TSharedPtr<FAnimNextEvaluationTas
 	GetProxyOnGameThread<FSequencerMixedAnimInstanceProxy>().SetMixerTask(InEvalTask);
 }
 
-void USequencerMixedAnimInstance::SetRootMotion(const TSharedPtr<FMovieSceneMixerRootMotionComponentData>& InRootMotionData)
-{
-	GetProxyOnGameThread<FSequencerMixedAnimInstanceProxy>().SetRootMotion(InRootMotionData);
-}
-
 void FSequencerMixedAnimInstanceProxy::Initialize(UAnimInstance* InAnimInstance)
 {
 	FAnimInstanceProxy::Initialize(InAnimInstance);
@@ -211,15 +206,6 @@ bool FSequencerMixedAnimInstanceProxy::Evaluate(FPoseContext& Output)
 				ResultPose.LODPose.CopyFrom(EvaluatedKeyframe->Pose);
 				ResultPose.Curves.CopyFrom(EvaluatedKeyframe->Curves);
 				ResultPose.Attributes.CopyFrom(EvaluatedKeyframe->Attributes);
-
-				// Re-bake the root motion transform into the pose if we don't have any other way of applying root motion
-				const FTransformAnimationAttribute* RootMotionAttribute = ResultPose.Attributes.Find<FTransformAnimationAttribute>(FAnimMixerComponentTypes::Get()->RootTransformAttributeId);
-				const bool bHasSwapRootBone = RootMotion != nullptr;
-				if (RootMotionAttribute && !bHasSwapRootBone)
-				{
-					ResultPose.LODPose.LocalTransformsView[0] = RootMotionAttribute->Value;
-				}
-
 				bHasValidOutput = true;
 			}
 
@@ -302,11 +288,6 @@ void FSequencerMixedAnimInstanceProxy::SetMixerTask(TSharedPtr<FAnimNextEvaluati
 	MixerTask = InEvalTask;
 }
 
-void FSequencerMixedAnimInstanceProxy::SetRootMotion(const TSharedPtr<FMovieSceneMixerRootMotionComponentData>&  InRootMotion)
-{
-	RootMotion = InRootMotion;
-}
-
 namespace UE::MovieScene
 {
 	/** Helper function to get our sequencer animation instance from a skeletal mesh component */
@@ -435,8 +416,7 @@ namespace UE::MovieScene
 			FMovieSceneEntityID EntityID,
 			FObjectComponent MeshComponent,
 			TInstancedStruct<FMovieSceneMixedAnimationTarget>& Target,
-			const TSharedPtr<FAnimNextEvaluationTask>& MixerTask,
-			const TSharedPtr<FMovieSceneMixerRootMotionComponentData>* RootMotionData
+			const TSharedPtr<FAnimNextEvaluationTask>& MixerTask
 		) const
 		{
 			// Invalid or incompatible target, return.
@@ -478,7 +458,6 @@ namespace UE::MovieScene
 			if (SequencerInstance)
 			{
 				SequencerInstance->SetMixerTask(MixerTask);
-				SequencerInstance->SetRootMotion(RootMotionData ? *RootMotionData : TSharedPtr<FMovieSceneMixerRootMotionComponentData>());
 			}
 			else if (!AnimInstanceTarget->HasFiredWarningForTarget())
 			{
@@ -588,7 +567,6 @@ void UMovieSceneAnimInstanceTargetSystem::OnSchedulePersistentTasks(UE::MovieSce
 		.Read(AnimMixerComponents->MeshComponent)
 		.Write(AnimMixerComponents->Target)
 		.Read(AnimMixerComponents->MixerTask)
-		.ReadOptional(AnimMixerComponents->MixerRootMotion)
 		.SetParams(Params)
 		.Schedule_PerEntity<FEvaluateAndApplyAnimationTasks>(&Linker->EntityManager, TaskScheduler,
 			Linker, this);
diff --git a/Engine/Plugins/Experimental/MovieSceneAnimMixer/Source/MovieSceneAnimMixer/Private/Systems/MovieSceneAnimMixerSystem.cpp b/Engine/Plugins/Experimental/MovieSceneAnimMixer/Source/MovieSceneAnimMixer/Private/Systems/MovieSceneAnimMixerSystem.cpp
--- a/Engine/Plugins/Experimental/MovieSceneAnimMixer/Source/MovieSceneAnimMixer/Private/Systems/MovieSceneAnimMixerSystem.cpp
+++ b/Engine/Plugins/Experimental/MovieSceneAnimMixer/Source/MovieSceneAnimMixer/Private/Systems/MovieSceneAnimMixerSystem.cpp
@@ -2,6 +2,8 @@
 
 #include "Systems/MovieSceneAnimMixerSystem.h"
 
+#include "MovieSceneRootMotionSection.h"
+
 #include "EntitySystem/BuiltInComponentTypes.h"
 #include "AnimMixerComponentTypes.h"
 #include "EntitySystem/MovieSceneEntityIDs.h"
@@ -35,6 +37,8 @@
 #include "Systems/MovieSceneAnimInstanceTargetSystem.h"
 #include "Systems/MovieSceneAnimNextTargetSystem.h"
 #include "Systems/MovieSceneAnimBlueprintTargetSystem.h"
+#include "Systems/MovieSceneComponentTransformSystem.h"
+#include "Systems/MovieSceneRootMotionSystem.h"
 #include "AnimSubsystem_SequencerMixer.h"
 
 #include UE_INLINE_GENERATED_CPP_BY_NAME(MovieSceneAnimMixerSystem)
@@ -44,204 +48,203 @@ namespace UE::MovieScene
 {
 	struct FUpdateTaskEntities
 	{
-		FUpdateTaskEntities(UMovieSceneAnimMixerSystem* InSystem)
-			:System(InSystem)
-
-		{
-
-		}
-		void ForEachEntity(FMovieSceneEntityID EntityID,
-			FObjectKey BoundObjectKey,
-			const TInstancedStruct<FMovieSceneMixedAnimationTarget>& Target,
+		static void ForEachEntity(
 			TSharedPtr<FAnimNextEvaluationTask> Task,
-			const double* WeightAndEasings) const
+			TSharedPtr<FMovieSceneAnimMixerEntry> MixerEntry,
+			const double* WeightAndEasings)
 		{
-			FMovieSceneAnimMixerKey Key({ BoundObjectKey, Target });
-			if (FMovieSceneAnimMixer* Mixer = System->Mixers.Find(Key))
-			{
-				// Find the entry in question
-				if (FMovieSceneAnimMixerEntry* Entry = Mixer->Entries.FindByPredicate([EntityID](const FMovieSceneAnimMixerEntry& TestEntry) { return TestEntry.EntityID == EntityID; }))
-				{
-					Entry->EvalTask = Task;
-					Entry->PoseWeight = WeightAndEasings ? *WeightAndEasings : 1.0;
-				}
-			}
+			MixerEntry->EvalTask = Task;
+			MixerEntry->PoseWeight = WeightAndEasings ? *WeightAndEasings : 1.0;
 		}
-
-		TObjectPtr<UMovieSceneAnimMixerSystem> System;
 	};
 
 	struct FEvaluateAnimMixers
 	{
-		FEvaluateAnimMixers(UMovieSceneAnimMixerSystem* InSystem)
-			:System(InSystem)
-		{
+		TMap<FMovieSceneAnimMixerKey, TSharedPtr<FMovieSceneAnimMixer>>* Mixers;
 
+		FEvaluateAnimMixers(TMap<FMovieSceneAnimMixerKey, TSharedPtr<FMovieSceneAnimMixer>>* InMixers)
+			: Mixers(InMixers)
+		{
 		}
 
-		void ForEachEntity(FMovieSceneEntityID EntityID,
+		void ForEachEntity(
 			FObjectComponent MeshComponent,
 			const TInstancedStruct<FMovieSceneMixedAnimationTarget>& Target,
 			TSharedPtr<FAnimNextEvaluationTask>& MixerTask) const
 		{
-			if (FMovieSceneAnimMixer* Mixer = System->Mixers.Find(FMovieSceneAnimMixerKey({ MeshComponent.GetObject(), Target })))
+			TSharedPtr<FMovieSceneAnimMixer> Mixer = Mixers->FindRef(FMovieSceneAnimMixerKey({ MeshComponent.GetObject(), Target }));
+			if (!Mixer)
 			{
-				// Create a new eval program
-				Mixer->EvaluationProgram = MakeShared<UE::AnimNext::FEvaluationProgram>();
+				return;
+			}
 
-				if (!MixerTask.IsValid())
-				{
-					MixerTask = MakeShared<FAnimNextExecuteProgramTask>();
-				}
-				TSharedPtr<FAnimNextExecuteProgramTask> EvalProgramTask = StaticCastSharedPtr<FAnimNextExecuteProgramTask>(MixerTask);
-				EvalProgramTask->Program = Mixer->EvaluationProgram;
-
-				// From lowest to highest priority, add the tasks to the program, grouped by their priority.
-				// Tasks within the same priority are blending using a weighted average, which is then collapsed into a single pose, and blended with the
-				//      subsequnt priority based on its total accumulated PoseWeight.
-				//
-				// For example, 5 animations with a structure of:
-				// Priority 0: (A, w:0.5), (B, w:0.5), (C, w:0.5)
-				// Priority 2: (D, w:0.25), (E, w:0.5)
-				//
-				// Would result in a blend stack equivalent to the following operation:
-				// P1 = A*(0.5/1.5) + B*(0.5/1.5) + C(0.5/1.5)
-				// P2 = D*(0.25/0.75) + B*(0.5/0.75)
-				//
-				// FinalPose = (0.25*P1) + (0.75*P2)
-				//
-				// Note that we assume here that the target will push some kind of 'base pose' to the VM before evaluating the task. This allows us to blend in/out from gameplay for example.
-
-				int32 LastAbsoluteBlendPriority = MIN_int32;
-				float AccumulatedAbsoluteBlendPriortyGroupWeight = 0.f;
-
-				const int32 NumMixerEntries = Mixer->Entries.Num();
-
-				// Function that performs a lookahead accumulation of the total number of absolute weights for the specified start index's priority,
-				//  and returns a value indicating whether the entries need a separate blend stack (true) or whether they can be blended as a single operation
-				auto LookaheadComputeAccumulatedWeight = [&Mixer, NumMixerEntries](int32 StartAtIndex, int32 Priority, float& OutAccumulatedWeight) -> bool
+			// Create a new eval program
+			Mixer->EvaluationProgram = MakeShared<UE::AnimNext::FEvaluationProgram>();
+
+			if (!MixerTask.IsValid())
+			{
+				MixerTask = MakeShared<FAnimNextExecuteProgramTask>();
+			}
+			TSharedPtr<FAnimNextExecuteProgramTask> EvalProgramTask = StaticCastSharedPtr<FAnimNextExecuteProgramTask>(MixerTask);
+			EvalProgramTask->Program = Mixer->EvaluationProgram;
+
+			// From lowest to highest priority, add the tasks to the program, grouped by their priority.
+			// Tasks within the same priority are blending using a weighted average, which is then collapsed into a single pose, and blended with the
+			//      subsequnt priority based on its total accumulated PoseWeight.
+			//
+			// For example, 5 animations with a structure of:
+			// Priority 0: (A, w:0.5), (B, w:0.5), (C, w:0.5)
+			// Priority 2: (D, w:0.25), (E, w:0.5)
+			//
+			// Would result in a blend stack equivalent to the following operation:
+			// P1 = A*(0.5/1.5) + B*(0.5/1.5) + C(0.5/1.5)
+			// P2 = D*(0.25/0.75) + B*(0.5/0.75)
+			//
+			// FinalPose = (0.25*P1) + (0.75*P2)
+			//
+			// Note that we assume here that the target will push some kind of 'base pose' to the VM before evaluating the task. This allows us to blend in/out from gameplay for example.
+
+			const int32 NumMixerEntries = Mixer->WeakEntries.Num();
+
+			// Function that performs a lookahead accumulation of the total number of absolute weights for the specified start index's priority,
+			//  and returns a value indicating whether the entries need a separate blend stack (true) or whether they can be blended as a single operation
+			auto LookaheadComputeAccumulatedWeight = [&Mixer, NumMixerEntries](int32 StartAtIndex, int32 Priority, float& OutAccumulatedWeight) -> bool
+			{
+				int32 NumAbsoluteBlends = 0;
+				int32 NumSkippedBlends = 0;
+				float AccumulatedWeight = 0.f;
+				for (int32 NextIndex = StartAtIndex; NextIndex < NumMixerEntries; ++NextIndex)
 				{
-					int32 NumAbsoluteBlends = 0;
-					int32 NumSkippedBlends = 0;
-					float AccumulatedWeight = 0.f;
-					for (int32 NextIndex = StartAtIndex; NextIndex < NumMixerEntries; ++NextIndex)
+					TSharedPtr<FMovieSceneAnimMixerEntry> NextMixerEntry = Mixer->WeakEntries[NextIndex].Pin();
+					if (!ensure(NextMixerEntry))
 					{
-						const FMovieSceneAnimMixerEntry& NextMixerEntry = Mixer->Entries[NextIndex];
-						if (NextMixerEntry.Priority != Priority || NextMixerEntry.bAdditive)
+						continue;
+					}
+					if (NextMixerEntry->Priority != Priority || NextMixerEntry->bAdditive)
+					{
+						break;
+					}
+					else if (NextMixerEntry->bRequiresBlend && !NextMixerEntry->bAdditive)
+					{
+						if (!FMath::IsNearlyEqual(AccumulatedWeight, AccumulatedWeight + NextMixerEntry->PoseWeight, KINDA_SMALL_NUMBER))
 						{
-							break;
+							++NumAbsoluteBlends;
+							AccumulatedWeight += NextMixerEntry->PoseWeight;
 						}
-						else if (NextMixerEntry.bRequiresBlend && !NextMixerEntry.bAdditive)
+						else
 						{
-							if (!FMath::IsNearlyEqual(AccumulatedWeight, AccumulatedWeight + NextMixerEntry.PoseWeight, KINDA_SMALL_NUMBER))
-							{
-								++NumAbsoluteBlends;
-								AccumulatedWeight += NextMixerEntry.PoseWeight;
-							}
-							else
-							{
-								++NumSkippedBlends;
-							}
+							++NumSkippedBlends;
 						}
 					}
+				}
 
-					// When skipping blends, the accumulated weight still needs to be factored in for the absolute blends that aren't skipped,
-					// so if the total is greater than one, use a separate blend stack, and the accumulated weight.
-					if (NumAbsoluteBlends + NumSkippedBlends > 1 && !FMath::IsNearlyZero(AccumulatedWeight, KINDA_SMALL_NUMBER))
-					{
-						OutAccumulatedWeight = AccumulatedWeight;
-						return true;
-					}
-					return false;
-				};
+				// When skipping blends, the accumulated weight still needs to be factored in for the absolute blends that aren't skipped,
+				// so if the total is greater than one, use a separate blend stack, and the accumulated weight.
+				if (NumAbsoluteBlends + NumSkippedBlends > 1 && !FMath::IsNearlyZero(AccumulatedWeight, KINDA_SMALL_NUMBER))
+				{
+					OutAccumulatedWeight = AccumulatedWeight;
+					return true;
+				}
+				return false;
+			};
 
-				for (int32 Index = 0; Index < NumMixerEntries; )
+			TSharedPtr<FMovieSceneMixerRootMotionComponentData> RootMotion = Mixer->WeakRootMotion.Pin();
+
+			for (int32 Index = 0; Index < NumMixerEntries; )
+			{
+				TSharedPtr<FMovieSceneAnimMixerEntry> PeekEntry = Mixer->WeakEntries[Index].Pin();
+				if (!ensure(PeekEntry))
 				{
-					const int32 Priority = Mixer->Entries[Index].Priority;
+					continue;
+				}
+
+				const int32 Priority = PeekEntry->Priority;
 
-					// Recompute the accumulated weights for this priority group
-					float AccumulatedPriorityWeight = 1.f;
-					// If there is more than one pose in this priority with an absolute weight, we need to blend those together using a weighted average before applying additives and blending with the next priority or the base pose
-					const bool bNeedsSeparateWeightStack = LookaheadComputeAccumulatedWeight(Index, Priority, AccumulatedPriorityWeight);
+				// Recompute the accumulated weights for this priority group
+				float AccumulatedPriorityWeight = 1.f;
+				// If there is more than one pose in this priority with an absolute weight, we need to blend those together using a weighted average before applying additives and blending with the next priority or the base pose
+				const bool bNeedsSeparateWeightStack = LookaheadComputeAccumulatedWeight(Index, Priority, AccumulatedPriorityWeight);
 
-					bool bIsFirstAbsoluteBlend = true;
-					for ( ; Index < NumMixerEntries; ++Index)
+				bool bIsFirstAbsoluteBlend = true;
+				for ( ; Index < NumMixerEntries; ++Index)
+				{
+					const TSharedPtr<FMovieSceneAnimMixerEntry>& MixerEntry = Mixer->WeakEntries[Index].Pin();
+					if (!MixerEntry)
 					{
-						const FMovieSceneAnimMixerEntry& MixerEntry = Mixer->Entries[Index];
-						if (MixerEntry.Priority != Priority)
-						{
-							// Break and continue with the outer loop
-							break;
-						}
+						continue;
+					}
 
-						if (!MixerEntry.EvalTask.IsValid())
-						{
-							continue;
-						}
+					if (MixerEntry->Priority != Priority)
+					{
+						// Break and continue with the outer loop
+						break;
+					}
 
-						if (bNeedsSeparateWeightStack && FMath::IsNearlyZero(MixerEntry.PoseWeight, KINDA_SMALL_NUMBER))
-						{
-							continue;
-						}
+					if (!MixerEntry->EvalTask.IsValid())
+					{
+						continue;
+					}
 
-						// Evaluate the pose itself
-						Mixer->EvaluationProgram->AppendTaskPtr(MixerEntry.EvalTask);
+					if (bNeedsSeparateWeightStack && FMath::IsNearlyZero(MixerEntry->PoseWeight, KINDA_SMALL_NUMBER))
+					{
+						continue;
+					}
 
-						if (Mixer->RootMotion && MixerEntry.RootMotionSettings && MixerEntry.RootMotionSettings->RootMotionSpace == EMovieSceneRootMotionSpace::AnimationSpace)
+					// Evaluate the pose itself
+					Mixer->EvaluationProgram->AppendTaskPtr(MixerEntry->EvalTask);
+
+					if (RootMotion && MixerEntry->RootMotionSettings && MixerEntry->RootMotionSettings->RootMotionSpace == EMovieSceneRootMotionSpace::AnimationSpace)
+					{
+						// Need to convert root motion attribute to world space for blending, insert task to do that
+						Mixer->EvaluationProgram->AppendTask(FAnimNextConvertRootMotionToWorldSpaceTask::Make(RootMotion));
+					}
+
+					if (MixerEntry->bRequiresBlend)
+					{
+						if (MixerEntry->bAdditive)
 						{
-							// Need to convert root motion attribute to world space for blending, insert task to do that
-							Mixer->EvaluationProgram->AppendTask(FAnimNextConvertRootMotionToWorldSpaceTask::Make(Mixer->RootMotion));
+							// Add additive blend task.
+							Mixer->EvaluationProgram->AppendTask(FAnimNextApplyAdditiveKeyframeTask::Make(MixerEntry->PoseWeight));
 						}
-
-						if (MixerEntry.bRequiresBlend)
+						else if (bNeedsSeparateWeightStack)
 						{
-							if (MixerEntry.bAdditive)
-							{
-								// Add additive blend task.
-								Mixer->EvaluationProgram->AppendTask(FAnimNextApplyAdditiveKeyframeTask::Make(MixerEntry.PoseWeight));
-							}
-							else if (bNeedsSeparateWeightStack)
+							// Add absolute blend task. The first absolute task within a priority level needs to overwrite the pose with its weighted average,
+							//     subsequent poses are added to this. This in effect leaves us with a single pose on the stack at the end of this priority level
+							//     that includes the weighted average of all absolutes within this priority level, plus all absolutes
+							if (bIsFirstAbsoluteBlend)
 							{
-								// Add absolute blend task. The first absolute task within a priority level needs to overwrite the pose with its weighted average,
-								//     subsequent poses are added to this. This in effect leaves us with a single pose on the stack at the end of this priority level
-								//     that includes the weighted average of all absolutes within this priority level, plus all absolutes
-								if (bIsFirstAbsoluteBlend)
-								{
-									bIsFirstAbsoluteBlend = false;
-									Mixer->EvaluationProgram->AppendTask(FAnimNextBlendOverwriteKeyframeWithScaleTask::Make(MixerEntry.PoseWeight / AccumulatedPriorityWeight));
-								}
-								else
-								{
-									Mixer->EvaluationProgram->AppendTask(FMovieSceneAccumulateAbsoluteBlendTask::Make(MixerEntry.PoseWeight / AccumulatedPriorityWeight));
-								}
+								bIsFirstAbsoluteBlend = false;
+								Mixer->EvaluationProgram->AppendTask(FAnimNextBlendOverwriteKeyframeWithScaleTask::Make(MixerEntry->PoseWeight / AccumulatedPriorityWeight));
 							}
 							else
 							{
-								Mixer->EvaluationProgram->AppendTask(FAnimNextBlendTwoKeyframesTask::Make(MixerEntry.PoseWeight));
+								Mixer->EvaluationProgram->AppendTask(FMovieSceneAccumulateAbsoluteBlendTask::Make(MixerEntry->PoseWeight / AccumulatedPriorityWeight));
 							}
 						}
-					}
-					// If we have a previous amount to blend with the last priority, add the tasks to ensure that happens
-					if (bNeedsSeparateWeightStack)
-					{
-						Mixer->EvaluationProgram->AppendTask(FAnimNextNormalizeKeyframeRotationsTask());
-
-						// For now we always blend weighted averages with the next pose with a weight of 1.
-						//    Ultimately this should be controlled by its own weight within an anim mixer layer track.
-						constexpr float WeightedAverageBlendWeight = 1.f;
-						Mixer->EvaluationProgram->AppendTask(FAnimNextBlendTwoKeyframesTask::Make(WeightedAverageBlendWeight));
+						else
+						{
+							Mixer->EvaluationProgram->AppendTask(FAnimNextBlendTwoKeyframesTask::Make(MixerEntry->PoseWeight));
+						}
 					}
 				}
-
-				if (Mixer->RootMotion)
+				// If we have a previous amount to blend with the last priority, add the tasks to ensure that happens
+				if (bNeedsSeparateWeightStack)
 				{
-					// Add task to store the final root motion result
-					Mixer->EvaluationProgram->AppendTask(FAnimNextStoreRootTransformTask::Make(Mixer->RootMotion));
+					Mixer->EvaluationProgram->AppendTask(FAnimNextNormalizeKeyframeRotationsTask());
+
+					// For now we always blend weighted averages with the next pose with a weight of 1.
+					//    Ultimately this should be controlled by its own weight within an anim mixer layer track.
+					constexpr float WeightedAverageBlendWeight = 1.f;
+					Mixer->EvaluationProgram->AppendTask(FAnimNextBlendTwoKeyframesTask::Make(WeightedAverageBlendWeight));
 				}
 			}
-		}
 
-		TObjectPtr<UMovieSceneAnimMixerSystem> System;
+			if (RootMotion)
+			{
+				// Add task to store the final root motion result
+				Mixer->EvaluationProgram->AppendTask(FAnimNextStoreRootTransformTask::Make(RootMotion));
+			}
+		}
 	};
 }
 
@@ -311,6 +314,7 @@ UMovieSceneAnimMixerSystem::UMovieSceneAnimMixerSystem(const FObjectInitializer&
 		DefineImplicitPrerequisite(UMovieSceneGenericBoundObjectInstantiator::StaticClass(), GetClass());
 		DefineImplicitPrerequisite(UMovieSceneBoundSceneComponentInstantiator::StaticClass(), GetClass());
 		DefineImplicitPrerequisite(UWeightAndEasingEvaluatorSystem::StaticClass(), GetClass());
+		DefineImplicitPrerequisite(UMovieSceneComponentTransformSystem::StaticClass(), GetClass());
 		DefineComponentConsumer(GetClass(), BuiltInComponents->BoundObjectKey);
 		DefineComponentConsumer(GetClass(), BuiltInComponents->BoundObject);
 		DefineComponentConsumer(GetClass(), AnimMixerComponents->Task);
@@ -320,6 +324,52 @@ UMovieSceneAnimMixerSystem::UMovieSceneAnimMixerSystem(const FObjectInitializer&
 	}
 }
 
+TSharedPtr<FMovieSceneMixerRootMotionComponentData> UMovieSceneAnimMixerSystem::FindRootMotion(FMovieSceneAnimMixerKey InMixerKey) const
+{
+	TSharedPtr<FMovieSceneAnimMixer> Mixer = Mixers.FindRef(InMixerKey);
+
+	return Mixer ? Mixer->WeakRootMotion.Pin() : TSharedPtr<FMovieSceneMixerRootMotionComponentData>();
+}
+
+void UMovieSceneAnimMixerSystem::AssignRootMotion(FMovieSceneAnimMixerKey InMixerKey, TSharedPtr<FMovieSceneMixerRootMotionComponentData> RootMotion)
+{
+	TSharedPtr<FMovieSceneAnimMixer> Mixer = Mixers.FindRef(InMixerKey);
+
+	if (!Mixer && RootMotion)
+	{
+		Mixer = MakeShared<FMovieSceneAnimMixer>();
+		Mixers.Add(InMixerKey, Mixer);
+	}
+	if (Mixer)
+	{
+		Mixer->bNeedsResort |= (Mixer->WeakRootMotion.Pin() != RootMotion);
+		Mixer->WeakRootMotion = RootMotion;
+	}
+}
+
+void UMovieSceneAnimMixerSystem::InitializeAllRootMotion()
+{
+	for (auto It = Mixers.CreateIterator(); It; ++It)
+	{
+		TSharedPtr<FMovieSceneMixerRootMotionComponentData> RootMotion = It.Value()->WeakRootMotion.Pin();;
+		if (RootMotion)
+		{
+			RootMotion->Initialize();
+		}
+	}
+}
+
+bool UMovieSceneAnimMixerSystem::IsRelevantImpl(UMovieSceneEntitySystemLinker* InLinker) const
+{
+	return Mixers.Num() != 0;
+}
+
+void UMovieSceneAnimMixerSystem::OnLink()
+{
+	RootMotionSystem = Linker->LinkSystem<UMovieSceneRootMotionSystem>();
+	Linker->SystemGraph.AddReference(this, RootMotionSystem);
+}
+
 void UMovieSceneAnimMixerSystem::OnRun(FSystemTaskPrerequisites& InPrerequisites, FSystemSubsequentTasks& Subsequents)
 {
 	using namespace UE::MovieScene;
@@ -334,19 +384,23 @@ void UMovieSceneAnimMixerSystem::OnRun(FSystemTaskPrerequisites& InPrerequisites
 
 		// ----------------------------------------------------------------------------------------------------------------------------------------
 		// Step 1 - Remove expired mixer entries.
-
 		FEntityTaskBuilder()
-			.ReadEntityIDs()
-			.Read(BuiltInComponents->BoundObjectKey)
-			.Read(AnimMixerComponents->Target)
-			.FilterAll({ BuiltInComponents->Tags.NeedsUnlink, AnimMixerComponents->Task})
-			.Iterate_PerEntity(&Linker->EntityManager, [this](FMovieSceneEntityID EntityID, FObjectKey BoundObjectKey, const TInstancedStruct<FMovieSceneMixedAnimationTarget>& Target)
+		.Read(BuiltInComponents->BoundObjectKey)
+		.Write(AnimMixerComponents->MixerEntry)
+		.FilterAll({ BuiltInComponents->Tags.NeedsUnlink, AnimMixerComponents->Task })
+		.Iterate_PerEntity(&Linker->EntityManager,
+			[this](FObjectKey BoundObjectKey, TSharedPtr<FMovieSceneAnimMixerEntry>& MixerEntry)
+			{
+				if (MixerEntry)
 				{
-					if (FMovieSceneAnimMixer* Mixer = Mixers.Find(FMovieSceneAnimMixerKey({ BoundObjectKey, Target })))
+					if (TSharedPtr<FMovieSceneAnimMixer> Mixer = MixerEntry->WeakMixer.Pin())
 					{
-						Mixer->Entries.RemoveAll([&EntityID](const FMovieSceneAnimMixerEntry& Entry) { return Entry.EntityID == EntityID;});
+						Mixer->bNeedsResort = true;
 					}
-				});
+				}
+				MixerEntry = nullptr;
+			}
+		);
 
 		// ----------------------------------------------------------------------------------------------------------------------------------------
 		// Step 2 - Create new mixer entries for new anim tasks, gathering any new mixer entities that need to be created
@@ -356,6 +410,7 @@ void UMovieSceneAnimMixerSystem::OnRun(FSystemTaskPrerequisites& InPrerequisites
 		.ReadEntityIDs()
 		.Read(BuiltInComponents->BoundObjectKey)
 		.Write(AnimMixerComponents->Target)
+		.Write(AnimMixerComponents->MixerEntry)
 		.Read(AnimMixerComponents->Task)
 		.Read(AnimMixerComponents->Priority)
 		.ReadOptional(AnimMixerComponents->RootMotionSettings)
@@ -366,6 +421,7 @@ void UMovieSceneAnimMixerSystem::OnRun(FSystemTaskPrerequisites& InPrerequisites
 			FMovieSceneEntityID EntityID,
 			FObjectKey BoundObjectKey,
 			TInstancedStruct<FMovieSceneMixedAnimationTarget>& Target,
+			TSharedPtr<FMovieSceneAnimMixerEntry>& InOutMixerEntry,
 			TSharedPtr<FAnimNextEvaluationTask> Task,
 			int32 Priority,
 			const FMovieSceneRootMotionSettings* RootMotionSettings)
@@ -373,111 +429,141 @@ void UMovieSceneAnimMixerSystem::OnRun(FSystemTaskPrerequisites& InPrerequisites
 				// For new entities, we 'resolve' the animation target so if 'Automatic' is picked we choose the right one automatically.
 				Target = ResolveAnimationTarget(BoundObjectKey, Target);
 				FMovieSceneAnimMixerKey Key({ BoundObjectKey, Target });
-				FMovieSceneAnimMixer* Mixer = Mixers.Find(Key);
+				TSharedPtr<FMovieSceneAnimMixer> Mixer = Mixers.FindRef(Key);
 				if (!Mixer)
 				{
-					Mixer = &Mixers.Add({ Key, FMovieSceneAnimMixer()});
+					Mixer = MakeShared<FMovieSceneAnimMixer>();
+					Mixers.Add(Key, Mixer);
 					NewMixers.Add(Key);
 				}
-				if (Mixer)
+				else if (Mixer->WeakEntries.Num() == 0)
 				{
-					checkSlow(!Mixer->Entries.ContainsByPredicate([EntityID](const FMovieSceneAnimMixerEntry& Entry) { return Entry.EntityID == EntityID;}));
-					
-					const FComponentMask& Type = Linker->EntityManager.GetEntityType(EntityID);
-					FMovieSceneAnimMixerEntry NewEntry;
-					NewEntry.EntityID = EntityID;
-					NewEntry.EvalTask = Task;
-					NewEntry.Priority = Priority;
-					NewEntry.PoseWeight = 0;
-					if (RootMotionSettings)
+					NewMixers.Add(Key);
+				}
+
+				// Create a new mixer entry if necessary
+				if (!InOutMixerEntry)
+				{
+					InOutMixerEntry = MakeShared<FMovieSceneAnimMixerEntry>();
+				}
+
+				FMovieSceneAnimMixerEntry* Entry = InOutMixerEntry.Get();
+
+				Entry->EntityID = EntityID;
+				Entry->EvalTask = Task;
+				Entry->Priority = Priority;
+				Entry->PoseWeight = 0;
+				if (RootMotionSettings)
+				{
+					Entry->RootMotionSettings = *RootMotionSettings;
+
+					// If we know we'll need root motion, ensure it is set up correctly with a lifetime reference that keeps it
+					//     alive as long as this entry
+					TSharedPtr<FMovieSceneMixerRootMotionComponentData> RootMotion = Mixer->WeakRootMotion.Pin();
+					if (!RootMotion)
 					{
-						NewEntry.RootMotionSettings = *RootMotionSettings;
+						RootMotion = MakeShared<FMovieSceneMixerRootMotionComponentData>();
+						RootMotion->RootDestination = EMovieSceneRootMotionDestination::RootBone;
+
+						Mixer->WeakRootMotion = RootMotion;
 					}
-					NewEntry.bAdditive = Type.ContainsAny({BuiltInComponents->Tags.AdditiveBlend, BuiltInComponents->Tags.AdditiveFromBaseBlend});
-					NewEntry.bRequiresBlend = Type.Contains(AnimMixerComponents->Tags.RequiresBlending);
-					int32 EntryIndex = Algo::LowerBound(Mixer->Entries, NewEntry);
-					Mixer->Entries.Insert(NewEntry, EntryIndex);
+					Entry->RootMotionLifetimeReference = RootMotion;
 				}
+
+				const FComponentMask& Type = Linker->EntityManager.GetEntityType(EntityID);
+				Entry->bAdditive = Type.ContainsAny({BuiltInComponents->Tags.AdditiveBlend, BuiltInComponents->Tags.AdditiveFromBaseBlend});
+				Entry->bRequiresBlend = Type.Contains(AnimMixerComponents->Tags.RequiresBlending);
+
+				TSharedPtr<FMovieSceneAnimMixer> ExistingMixer = Entry->WeakMixer.Pin();
+				if (ExistingMixer)
+				{
+					if (ExistingMixer != Mixer)
+					{
+						ExistingMixer->WeakEntries.Remove(InOutMixerEntry);
+						Entry->WeakMixer = nullptr;
+
+						Mixer->WeakEntries.Emplace(InOutMixerEntry);
+						Entry->WeakMixer = Mixer;
+					}
+				}
+				else
+				{
+					Mixer->WeakEntries.Emplace(InOutMixerEntry);
+					Entry->WeakMixer = Mixer;
+				}
+
+				Mixer->bNeedsResort = true;
 			});
 
 		// ----------------------------------------------------------------------------------------------------------------------------------------
 		// Step 3 - Create new mixer entities
-		for (FMovieSceneAnimMixerKey NewMixerKey : NewMixers)
+		for (const FMovieSceneAnimMixerKey& NewMixerKey : NewMixers)
 		{
-			FMovieSceneAnimMixer& NewMixer = Mixers.FindChecked(NewMixerKey);
-			NewMixer.MixerEntityID = FEntityBuilder()
+			TSharedPtr<FMovieSceneAnimMixer> NewMixer = Mixers.FindChecked(NewMixerKey);
+			NewMixer->MixerEntityID = FEntityBuilder()
 				.Add(AnimMixerComponents->MeshComponent, FObjectComponent::Weak(NewMixerKey.BoundObjectKey.ResolveObjectPtr()))
 				.Add(AnimMixerComponents->Target, NewMixerKey.Target)
 				.Add(AnimMixerComponents->MixerTask, nullptr)
+				.AddConditional(AnimMixerComponents->MixerRootMotion, NewMixer->WeakRootMotion.Pin(), NewMixer->WeakRootMotion.Pin().IsValid())
 				.AddTag(BuiltInComponents->Tags.RestoreState) // TODO: For now we always restore state on the mixer when it gets unlinked.
 				.CreateEntity(&Linker->EntityManager);
 		}
 
+		auto SortPredicate = [](const TWeakPtr<FMovieSceneAnimMixerEntry>& A, const TWeakPtr<FMovieSceneAnimMixerEntry>& B)
+		{
+			return *A.Pin() < *B.Pin();
+		};
 
 		// ----------------------------------------------------------------------------------------------------------------------------------------
 		// Step 4 - Update mixer entities, and remove stale ones. 
-		// We also add/remove root motion components from the mixer entities depending on the entries present in them.
-		TArray<FMovieSceneAnimMixerKey> StaleMixers;
-		TArray<FObjectKey> StaleRootMotion;
-		TArray<TSharedPtr<FMovieSceneMixerRootMotionComponentData>> NewRootMotion;
-		for (TPair<FMovieSceneAnimMixerKey, FMovieSceneAnimMixer>& MixerPair : Mixers)
+		for (auto It = Mixers.CreateIterator(); It; ++It)
 		{
-			ESwapRootBone SwapRootBone = ESwapRootBone::SwapRootBone_None;
-			FMovieSceneAnimMixer& Mixer = MixerPair.Value;
-			const FMovieSceneAnimMixerEntry* MaxEntry = Algo::MaxElementBy(Mixer.Entries, [](const FMovieSceneAnimMixerEntry& Entry) { return Entry.RootMotionSettings ? Entry.RootMotionSettings->SwapRootBone : ESwapRootBone::SwapRootBone_None; });
-			if (MaxEntry)
+			TSharedPtr<FMovieSceneAnimMixer> Mixer = It.Value();
+
+			if (Mixer->bNeedsResort)
 			{
-				SwapRootBone = MaxEntry->RootMotionSettings ? MaxEntry->RootMotionSettings->SwapRootBone : ESwapRootBone::SwapRootBone_None;
-			}
+				bool bNeedsRootMotion = false;
 
-			bool bHasRootMotionComponent = Linker->EntityManager.GetEntityType(Mixer.MixerEntityID).Contains(AnimMixerComponents->MixerRootMotion);
+				// Remove nullptrs
+				for (int32 Index = Mixer->WeakEntries.Num()-1; Index >= 0; --Index)
+				{
+					if (Mixer->WeakEntries[Index].Pin() == nullptr)
+					{
+						Mixer->WeakEntries.RemoveAtSwap(Index, 1, EAllowShrinking::No);
+					}
+				}
 
-			if (Mixer.Entries.IsEmpty())
-			{
-				StaleMixers.Add(MixerPair.Key);
-				Linker->EntityManager.AddComponent(Mixer.MixerEntityID, BuiltInComponents->Tags.NeedsUnlink, EEntityRecursion::Full);
+				Algo::Sort(Mixer->WeakEntries, SortPredicate);
+				Mixer->bNeedsResort = false;
+
+				if (Mixer->MixerEntityID)
+				{
+					if (TSharedPtr<FMovieSceneMixerRootMotionComponentData> RootMotion = Mixer->WeakRootMotion.Pin())
+					{
+						Linker->EntityManager.AddComponent(Mixer->MixerEntityID, AnimMixerComponents->MixerRootMotion, RootMotion);
+					}
+					else
+					{
+						Linker->EntityManager.RemoveComponent(Mixer->MixerEntityID, AnimMixerComponents->MixerRootMotion);
+					}
+				}
 			}
-			else if (SwapRootBone != ESwapRootBone::SwapRootBone_None)
+
+			if (Mixer->WeakEntries.IsEmpty())
 			{
-				if (!Mixer.RootMotion.IsValid())
+				if (Mixer->MixerEntityID)
 				{
-					Mixer.RootMotion = MakeShared<FMovieSceneMixerRootMotionComponentData>();
+					Linker->EntityManager.AddComponent(Mixer->MixerEntityID, BuiltInComponents->Tags.NeedsUnlink, EEntityRecursion::Full);
+					Mixer->MixerEntityID = FMovieSceneEntityID();
 				}
-				Mixer.RootMotion->MixerEntityID = Mixer.MixerEntityID;
-				
-				// Ensure actor matches
-				USkeletalMeshComponent* SkeletalMeshComponent = Cast<USkeletalMeshComponent>(MixerPair.Key.BoundObjectKey.ResolveObjectPtr());
-				USceneComponent* RootMotionTarget = SkeletalMeshComponent;
-				if (SwapRootBone == ESwapRootBone::SwapRootBone_Actor)
+
+				// If we don't have any root motion then we can clean this up as well
+				if (Mixer->WeakRootMotion.Pin() == nullptr)
 				{
-					RootMotionTarget = SkeletalMeshComponent ? SkeletalMeshComponent->GetOwner()->GetRootComponent() : nullptr;
+					It.RemoveCurrent();
 				}
-				Mixer.RootMotion->Target = RootMotionTarget;
-
-				if (Mixer.RootMotion->Target.IsValid() && !bHasRootMotionComponent)
-				{ 
-					// Add this root motion to the mixer entity
-					FEntityBuilder()
-						.Add(AnimMixerComponents->MixerRootMotion, Mixer.RootMotion)
-						.AddTag(BuiltInComponents->Tags.NeedsLink)
-						.MutateExisting(&Linker->EntityManager, Mixer.MixerEntityID);
-				}
-			}
-			else if (SwapRootBone == ESwapRootBone::SwapRootBone_None && bHasRootMotionComponent)
-			{
-				Mixer.RootMotion = nullptr;
-
-				// Remove root motion from the mixer entity
-				FTypelessMutation Mutation;
-				Mutation.Remove({ AnimMixerComponents->MixerRootMotion });
-				FEntityBuilder()
-					.AddTag(BuiltInComponents->Tags.NeedsLink)
-					.MutateExisting(&Linker->EntityManager, Mixer.MixerEntityID, Mutation);
 			}
 		}
-
-		// Remove stale mixers
-		Algo::ForEach(StaleMixers, [this](const FMovieSceneAnimMixerKey& MixerKey) { Mixers.Remove(MixerKey); });
 	}
 }
 
@@ -495,23 +581,19 @@ void UMovieSceneAnimMixerSystem::OnSchedulePersistentTasks(UE::MovieScene::IEnti
 	// Update mixer entry tasks- todo this is slow with the current data hierarchy
 	// Maybe we want a flat map of entity id to entry, and then put indices into that into the mixer or something similar.
 	FTaskID UpdateTask = FEntityTaskBuilder()
-		.ReadEntityIDs()
-		.Read(BuiltInComponents->BoundObjectKey)
-		.Read(AnimMixerComponents->Target)
 		.Read(AnimMixerComponents->Task)
+		.Write(AnimMixerComponents->MixerEntry)
 		.ReadOptional(BuiltInComponents->WeightAndEasingResult)
 		.FilterNone({BuiltInComponents->Tags.NeedsUnlink})
-		.Schedule_PerEntity<FUpdateTaskEntities>(&Linker->EntityManager, TaskScheduler, this );
+		.Schedule_PerEntity<FUpdateTaskEntities>(&Linker->EntityManager, TaskScheduler);
 
 
 	// For each mixer, build the evaluation program task list.
 	FTaskID MixTask = FEntityTaskBuilder()
-		.ReadEntityIDs()
 		.Read(AnimMixerComponents->MeshComponent)
 		.Read(AnimMixerComponents->Target)
 		.Write(AnimMixerComponents->MixerTask)
-		.FilterAll({ AnimMixerComponents->MixerTask })
-		.Schedule_PerEntity<FEvaluateAnimMixers>(&Linker->EntityManager, TaskScheduler, this);
+		.Schedule_PerEntity<FEvaluateAnimMixers>(&Linker->EntityManager, TaskScheduler, &Mixers);
 
 	TaskScheduler->AddPrerequisite(UpdateTask, MixTask);
 }
@@ -524,33 +606,37 @@ void UMovieSceneAnimMixerSystem::OnCleanTaggedGarbage()
 	FAnimMixerComponentTypes* AnimMixerComponents = FAnimMixerComponentTypes::Get();
 
 	FEntityTaskBuilder()
-		.ReadEntityIDs()
-		.Read(BuiltInComponents->BoundObjectKey)
-		.Read(AnimMixerComponents->Target)
-		.FilterAll({ BuiltInComponents->Tags.NeedsUnlink, AnimMixerComponents->Task })
-		.Iterate_PerEntity(&Linker->EntityManager, [this](FMovieSceneEntityID EntityID, FObjectKey BoundObjectKey, const TInstancedStruct<FMovieSceneMixedAnimationTarget>& Target)
+	.Read(BuiltInComponents->BoundObjectKey)
+	.Write(AnimMixerComponents->MixerEntry)
+	.FilterAll({ BuiltInComponents->Tags.NeedsUnlink, AnimMixerComponents->Task })
+	.Iterate_PerEntity(&Linker->EntityManager,
+		[this](FObjectKey BoundObjectKey, TSharedPtr<FMovieSceneAnimMixerEntry>& MixerEntry)
+		{
+			if (MixerEntry)
 			{
-				if (FMovieSceneAnimMixer* Mixer = Mixers.Find(FMovieSceneAnimMixerKey({ BoundObjectKey, Target })))
+				TSharedPtr<FMovieSceneAnimMixer> Mixer = MixerEntry->WeakMixer.Pin();
+				if (Mixer)
 				{
-					Mixer->Entries.RemoveAll([&EntityID](const FMovieSceneAnimMixerEntry& Entry) { return Entry.EntityID == EntityID; });
+					Mixer->bNeedsResort = true;
 				}
-			});
+			}
+			MixerEntry = nullptr;
+		}
+	);
 
 	FEntityTaskBuilder()
-		.ReadEntityIDs()
-		.Read(AnimMixerComponents->MeshComponent)
-		.Read(AnimMixerComponents->Target)
-		.FilterAll({ BuiltInComponents->Tags.NeedsUnlink, AnimMixerComponents->MixerTask })
-		.Iterate_PerEntity(&Linker->EntityManager, [this]
-			(FMovieSceneEntityID EntityID,
-			FObjectComponent MeshComponent,
-			const TInstancedStruct<FMovieSceneMixedAnimationTarget>& Target)
-			{
-				Mixers.Remove(FMovieSceneAnimMixerKey({ MeshComponent.GetObject(), Target }));
-			});
+	.Read(AnimMixerComponents->MeshComponent)
+	.Read(AnimMixerComponents->Target)
+	.FilterAll({ BuiltInComponents->Tags.NeedsUnlink, AnimMixerComponents->MixerTask })
+	.Iterate_PerEntity(&Linker->EntityManager,
+		[this](FObjectComponent MeshComponent, const TInstancedStruct<FMovieSceneMixedAnimationTarget>& Target)
+		{
+			Mixers.Remove(FMovieSceneAnimMixerKey({ MeshComponent.GetObject(), Target }));
+		}
+	);
 }
 
-const TInstancedStruct<FMovieSceneMixedAnimationTarget> UMovieSceneAnimMixerSystem::ResolveAnimationTarget(FObjectKey ObjectKey, const TInstancedStruct<FMovieSceneMixedAnimationTarget> InTarget)
+TInstancedStruct<FMovieSceneMixedAnimationTarget> UMovieSceneAnimMixerSystem::ResolveAnimationTarget(FObjectKey ObjectKey, const TInstancedStruct<FMovieSceneMixedAnimationTarget>& InTarget)
 {
 	// If user has selected the default 'automatic' target, attempt to choose one automatically for them.
 	if (!InTarget.IsValid() || InTarget.GetScriptStruct() == FMovieSceneMixedAnimationTarget::StaticStruct())
diff --git a/Engine/Plugins/Experimental/MovieSceneAnimMixer/Source/MovieSceneAnimMixer/Private/Systems/MovieSceneAnimNextTargetSystem.cpp b/Engine/Plugins/Experimental/MovieSceneAnimMixer/Source/MovieSceneAnimMixer/Private/Systems/MovieSceneAnimNextTargetSystem.cpp
--- a/Engine/Plugins/Experimental/MovieSceneAnimMixer/Source/MovieSceneAnimMixer/Private/Systems/MovieSceneAnimNextTargetSystem.cpp
+++ b/Engine/Plugins/Experimental/MovieSceneAnimMixer/Source/MovieSceneAnimMixer/Private/Systems/MovieSceneAnimNextTargetSystem.cpp
@@ -109,7 +109,7 @@ namespace UE::MovieScene
 		void RestorePreAnimatedValue(const KeyType& Object, FSkeletalMeshEditorParams& InOutCachedValue, const FRestoreStateParams& Params)
 		{
 			USkeletalMeshComponent* SkelMeshComp = Cast<USkeletalMeshComponent>(Object.ResolveObjectPtr());
-			if (SkelMeshComp != nullptr)
+			if (SkelMeshComp != nullptr && SkelMeshComp->GetSceneProxy())
 			{
 				SkelMeshComp->GetSceneProxy()->SetCanSkipRedundantTransformUpdates(true);
 				InOutCachedValue.RestoreState();
diff --git a/Engine/Plugins/Experimental/MovieSceneAnimMixer/Source/MovieSceneAnimMixer/Private/Systems/MovieSceneMixedSkeletalAnimationSystem.cpp b/Engine/Plugins/Experimental/MovieSceneAnimMixer/Source/MovieSceneAnimMixer/Private/Systems/MovieSceneMixedSkeletalAnimationSystem.cpp
--- a/Engine/Plugins/Experimental/MovieSceneAnimMixer/Source/MovieSceneAnimMixer/Private/Systems/MovieSceneMixedSkeletalAnimationSystem.cpp
+++ b/Engine/Plugins/Experimental/MovieSceneAnimMixer/Source/MovieSceneAnimMixer/Private/Systems/MovieSceneMixedSkeletalAnimationSystem.cpp
@@ -75,6 +75,7 @@ struct FSkeletalAnimMixerMutation : IMovieSceneEntityMutation
 			AnimMixerComponents->Priority,
 			AnimMixerComponents->Target,
 			AnimMixerComponents->Task,
+			AnimMixerComponents->MixerEntry,
 			AnimMixerComponents->RootMotionSettings,
 			AnimMixerComponents->Tags.RequiresBlending
 		});
@@ -109,22 +110,32 @@ struct FSkeletalAnimMixerMutation : IMovieSceneEntityMutation
 				}
 
 				OutPriorities[Index] = SkeletalAnim.Section->MixedAnimationPriority;
-				OutRootMotionSettings[Index].SwapRootBone = SkeletalAnim.Section->Params.SwapRootBone;
 				OutRootMotionSettings[Index].RootMotionSpace = EMovieSceneRootMotionSpace::AnimationSpace;
+				OutRootMotionSettings[Index].LegacySwapRootBone = SkeletalAnim.Section->Params.SwapRootBone;
 
+				bool bHasRootMotionOverride = false;
 				UMovieSceneSkeletalAnimationTrack* Track = SkeletalAnim.Section->GetTypedOuter<UMovieSceneSkeletalAnimationTrack>();
-
-				if (Track->RootMotionParams.bRootMotionsDirty) 
+				if (Track)
 				{
-					constexpr bool bForce = true;
-					Track->SetUpRootMotions(bForce);
-				}
+					if (OutRootMotionSettings[Index].LegacySwapRootBone == ESwapRootBone::SwapRootBone_None)
+					{
+						OutRootMotionSettings[Index].LegacySwapRootBone = Track->SwapRootBone;
+					}
 
-				if (Track->RootMotionParams.bHaveRootMotion)
-				{
-					OutRootMotionSettings[Index].RootMotionOverride.Emplace();
+					if (Track->RootMotionParams.bRootMotionsDirty) 
+					{
+						constexpr bool bForce = true;
+						Track->SetUpRootMotions(bForce);
+					}
+
+					if (Track->RootMotionParams.bHaveRootMotion)
+					{
+						bHasRootMotionOverride = true;
+						OutRootMotionSettings[Index].RootMotionOverride.Emplace();
+					}
 				}
-				else
+
+				if (!bHasRootMotionOverride)
 				{
 					OutRootMotionSettings[Index].RootMotionOverride.Reset();
 				}
@@ -180,6 +191,8 @@ struct FGatherMixableSkeletalAnimations
 		TRead<FMovieSceneSkeletalAnimationComponentData> SkeletalAnimations,
 		TReadOptional<FFrameTime> OptionalEvalTimes,
 		TReadOptional<FMovieSceneRootMotionSettings> OptionalRootMotionSettings,
+		const double* OptLocationX, const double* OptLocationY, const double* OptLocationZ,
+		const double* OptRotationX, const double* OptRotationY, const double* OptRotationZ,
 		TWrite<TSharedPtr<FAnimNextEvaluationTask>> EvalTask) const
 	{
 		// Gather all the skeletal animations currently active in all sequences.
@@ -188,6 +201,7 @@ struct FGatherMixableSkeletalAnimations
 		const FBuiltInComponentTypes* BuiltInComponents = FBuiltInComponentTypes::Get();
 		const FEntityAllocation* Allocation = AllocationProxy.GetAllocation();
 		const int32 Num = Allocation->Num();
+
 		for (int32 Index = 0; Index < Num; ++Index)
 		{
 			FMovieSceneEntityID EntityID(EntityIDs[Index]);
@@ -284,9 +298,12 @@ struct FGatherMixableSkeletalAnimations
 			AnimData.FromPosition = PreviousEvalTime;
 			AnimData.ToPosition = EvalTime;
 			AnimData.bFireNotifies = bFireNotifies && !AnimParams.bSkipAnimNotifiers && !bLooped;
-			AnimData.SwapRootBone = AnimParams.SwapRootBone;
 			// TODO: We want to allow for additive to be set on the section properties
 			AnimData.bAdditive = AnimData.AnimSequence.IsValid() ? AnimData.AnimSequence->IsValidAdditive() : false;
+			if (OptionalRootMotionSettings)
+			{
+				AnimData.LegacySwapRootBone = OptionalRootMotionSettings->LegacySwapRootBone;
+			}
 
 			// TODO: Need to create a mirroring task and add that if appropriate
 			// TODO: Notifies- perhaps this should be a part of the animnext side of things?
@@ -298,6 +315,19 @@ struct FGatherMixableSkeletalAnimations
 			}
 			TSharedPtr<FMovieSceneSkeletalAnimationEvaluationTask> AnimTask = StaticCastSharedPtr<FMovieSceneSkeletalAnimationEvaluationTask>(EvalTask[Index]);
 			AnimTask->AnimationData = AnimData;
+
+			FVector  CurrentLocation = AnimTask->RootBaseTransform.GetTranslation();
+			FRotator CurrentRotation = AnimTask->RootBaseTransform.GetRotation().Rotator();
+
+			if (OptLocationX) { CurrentLocation.X = OptLocationX[Index]; }
+			if (OptLocationY) { CurrentLocation.Y = OptLocationY[Index]; }
+			if (OptLocationZ) { CurrentLocation.Z = OptLocationZ[Index]; }
+			if (OptRotationX) { CurrentRotation.Roll  = OptRotationX[Index]; }
+			if (OptRotationY) { CurrentRotation.Pitch = OptRotationY[Index]; }
+			if (OptRotationZ) { CurrentRotation.Yaw   = OptRotationZ[Index]; }
+
+			AnimTask->RootBaseTransform.SetRotation(CurrentRotation.Quaternion());
+			AnimTask->RootBaseTransform.SetLocation(CurrentLocation);
 		}
 	}
 
@@ -325,12 +355,12 @@ void FMovieSceneSkeletalAnimationEvaluationTask::Execute(UE::AnimNext::FEvaluati
 		DeltaTime.Set(AnimationData.FromPosition, AnimationData.ToPosition);
 
 		const bool bExtractRootMotion = AnimationData.RootMotionOverride.IsSet();
-		const FAnimExtractContext ExtractionContext(AnimationData.ToPosition, bExtractRootMotion, DeltaTime, false);
+		FAnimExtractContext ExtractionContext(AnimationData.ToPosition, bExtractRootMotion, DeltaTime, false);
 
 		FKeyframeState Keyframe = VM.MakeUninitializedKeyframe(bIsAdditive);
 
 		const bool bUseRawData = FDecompressionTools::ShouldUseRawData(AnimSequencePtr, Keyframe.Pose);
-		const bool bUseRootMotion = AnimSequencePtr->bEnableRootMotion || AnimationData.SwapRootBone != ESwapRootBone::SwapRootBone_None;
+		const bool bUseRootMotion = AnimSequencePtr->bEnableRootMotion;
 
 		TOptional<FTransform> RootBoneTransform;
 
@@ -348,7 +378,7 @@ void FMovieSceneSkeletalAnimationEvaluationTask::Execute(UE::AnimNext::FEvaluati
 				else if (AnimationData.RootMotionOverride->ChildBoneIndex != INDEX_NONE)
 				{
 					PoseIndex = Keyframe.Pose.GetRefPose().GetLODBoneIndexFromSkeletonBoneIndex(AnimationData.RootMotionOverride->ChildBoneIndex);
-				}					
+				}
 				if (PoseIndex != INDEX_NONE)
 				{
 					Keyframe.Pose.LocalTransformsView[PoseIndex] = AnimationData.RootMotionOverride->RootMotion;
@@ -390,13 +420,13 @@ void FMovieSceneSkeletalAnimationEvaluationTask::Execute(UE::AnimNext::FEvaluati
 			}
 
 			// If we haven't already grabbed the root transform, grab it here in case we're not doing the full bone sample
-			if (AnimationData.SwapRootBone != ESwapRootBone::SwapRootBone_None && !RootBoneTransform.IsSet())
+			if (!RootBoneTransform.IsSet())
 			{
 				if (AnimationData.RootMotionOverride.IsSet())
 				{
 					RootBoneTransform = AnimationData.RootMotionOverride->RootMotion;
 				}
-				else
+				else if (bUseRootMotion)
 				{
 					RootBoneTransform = AnimSequencePtr->ExtractRootTrackTransform(FAnimExtractContext(AnimationData.ToPosition), nullptr);
 				}
@@ -405,8 +435,19 @@ void FMovieSceneSkeletalAnimationEvaluationTask::Execute(UE::AnimNext::FEvaluati
 
 		if (RootBoneTransform.IsSet())
 		{
-			FTransformAnimationAttribute* RootMotionAttribute = Keyframe.Attributes.FindOrAdd<FTransformAnimationAttribute>(AnimMixerComponents->RootTransformAttributeId);
-			RootMotionAttribute->Value = RootBoneTransform.GetValue();
+			FTransformAnimationAttribute* RootTransformBaseAttribute = Keyframe.Attributes.FindOrAdd<FTransformAnimationAttribute>(AnimMixerComponents->RootTransformBaseAttributeId);
+			FTransformAnimationAttribute* RootMotionAttribute        = Keyframe.Attributes.FindOrAdd<FTransformAnimationAttribute>(AnimMixerComponents->RootTransformAttributeId);
+			FFloatAnimationAttribute*     RootMotionWeight           = Keyframe.Attributes.FindOrAdd<FFloatAnimationAttribute>(AnimMixerComponents->RootTransformWeightAttributeId);
+
+			RootTransformBaseAttribute->Value = RootBaseTransform;
+			RootMotionAttribute->Value        = RootBoneTransform.GetValue();
+			RootMotionWeight->Value           = 1.f;
+
+			if (AnimationData.LegacySwapRootBone != ESwapRootBone::SwapRootBone_None)
+			{
+				FIntegerAnimationAttribute* Int = Keyframe.Attributes.FindOrAdd<FIntegerAnimationAttribute>(AnimMixerComponents->LegacySwapRootBoneAttributeId);
+				Int->Value = (int32)AnimationData.LegacySwapRootBone;
+			}
 		}
 
 		VM.PushValue(KEYFRAME_STACK_NAME, MakeUnique<FKeyframeState>(MoveTemp(Keyframe)));
@@ -500,6 +541,12 @@ void UMovieSceneMixedSkeletalAnimationSystem::OnSchedulePersistentTasks(UE::Movi
 		.Read(TrackComponents->SkeletalAnimation)
 		.ReadOptional(BuiltInComponents->EvalTime)
 		.ReadOptional(AnimMixerComponents->RootMotionSettings)
+		.ReadOptional(BuiltInComponents->DoubleResult[0])
+		.ReadOptional(BuiltInComponents->DoubleResult[1])
+		.ReadOptional(BuiltInComponents->DoubleResult[2])
+		.ReadOptional(BuiltInComponents->DoubleResult[3])
+		.ReadOptional(BuiltInComponents->DoubleResult[4])
+		.ReadOptional(BuiltInComponents->DoubleResult[5])
 		.Write(AnimMixerComponents->Task)
 		.FilterAll( { TrackComponents->Tags.AnimMixerPoseProducer })
 		.FilterNone({ BuiltInComponents->Tags.Ignored})
diff --git a/Engine/Plugins/Experimental/MovieSceneAnimMixer/Source/MovieSceneAnimMixer/Private/Systems/MovieSceneRootMotionSystem.cpp b/Engine/Plugins/Experimental/MovieSceneAnimMixer/Source/MovieSceneAnimMixer/Private/Systems/MovieSceneRootMotionSystem.cpp
--- a/Engine/Plugins/Experimental/MovieSceneAnimMixer/Source/MovieSceneAnimMixer/Private/Systems/MovieSceneRootMotionSystem.cpp
+++ b/Engine/Plugins/Experimental/MovieSceneAnimMixer/Source/MovieSceneAnimMixer/Private/Systems/MovieSceneRootMotionSystem.cpp
@@ -5,22 +5,30 @@
 #include "EntitySystem/BuiltInComponentTypes.h"
 #include "EntitySystem/IMovieSceneTaskScheduler.h"
 #include "MovieSceneTracksComponentTypes.h"
+#include "MovieSceneRootMotionSection.h"
 
 #include "EvaluationVM/EvaluationTask.h"
 #include "EvaluationVM/EvaluationVM.h"
+#include "Component/AnimNextComponent.h"
 
 #include "Systems/MovieSceneAnimMixerSystem.h"
 #include "AnimMixerComponentTypes.h"
+#include "Module/AnimNextModuleInstance.h"
+#include "Animation/AnimRootMotionProvider.h"
+
+#include "SceneInterface.h"
 
 #include "Components/SkeletalMeshComponent.h"
 #include "Graph/AnimNext_LODPose.h"
 #include "Animation/BuiltInAttributeTypes.h"
 #include "EntitySystem/Interrogation/MovieSceneInterrogationLinker.h"
 #include "Systems/MovieScenePiecewiseDoubleBlenderSystem.h"
+#include "Systems/DoubleChannelEvaluatorSystem.h"
 #include "Misc/MemStack.h"
 #include "EntitySystem/MovieSceneEntitySystemTask.h"
 #include "EntitySystem/MovieSceneEntitySystemLinker.h"
 #include "Systems/MovieScenePropertyInstantiator.h"
+#include "Systems/MovieSceneComponentTransformSystem.h"
 #include "EntitySystem/Interrogation/MovieSceneInterrogatedPropertyInstantiator.h"
 #include "EntitySystem/MovieSceneEntitySystemRunner.h"
 #include "Misc/ScopeRWLock.h"
@@ -35,137 +43,76 @@
 
 namespace UE::MovieScene
 {
-	struct FUpdateRootMotionData
+	struct FGatherRootDestinations
 	{
-		void ForEachAllocation(
-			const FEntityAllocation* Allocation,
-			TRead<FObjectComponent> MeshComponents,
-			TWrite<TSharedPtr<FMovieSceneMixerRootMotionComponentData>> RootMotions) const
+		static void ForEachEntity(
+			uint8 ByteResult,
+			EMovieSceneRootMotionDestination& OutRootDestination,
+			TSharedPtr<FMovieSceneMixerRootMotionComponentData>& RootMotion
+			)
 		{
-			using namespace UE::AnimNext;
-			using namespace UE::MovieScene;
+			OutRootDestination = (EMovieSceneRootMotionDestination)(ByteResult);
+			RootMotion->RootDestination = OutRootDestination;
+		}
+	};
 
-			for (int32 Index = 0; Index < Allocation->Num(); ++Index)
-			{
-				if (TSharedPtr<FMovieSceneMixerRootMotionComponentData> RootMotion = RootMotions[Index])
-				{
-					USkeletalMeshComponent* SkeletalMeshComponent = Cast<USkeletalMeshComponent>(MeshComponents[Index].GetObject());
-					if (SkeletalMeshComponent && SkeletalMeshComponent != SkeletalMeshComponent->GetOwner()->GetRootComponent() && SkeletalMeshComponent != RootMotion->Target)
-					{
-						RootMotion->SetInverseMeshToActorRotation(SkeletalMeshComponent->GetOwner()->GetRootComponent()->GetComponentTransform().GetRelativeTransformReverse(SkeletalMeshComponent->GetComponentTransform()).GetRotation());
-					}
-				}
-			}
+	struct FInitializeRootMotionTask
+	{
+		UMovieSceneAnimMixerSystem* AnimMixer;
+		
+		void Run(FEntityAllocationWriteContext WriteContext) const
+		{
+			AnimMixer->InitializeAllRootMotion();
 		}
+
 	};
 
-	struct FSyncTransformTrackResult
+	struct FIsActorBeingMovedStatics
 	{
-		UMovieSceneRootMotionSystem* System = nullptr;
+#if !WITH_EDITOR
+		constexpr bool IsActive() const
+		{
+			return false;
+		}
 
-		FMovieSceneTracksComponentTypes* TrackComponents = nullptr;
-		
-		FSyncTransformTrackResult(UMovieSceneRootMotionSystem* InRootMotionSystem)
-		: System(InRootMotionSystem)
+		constexpr void Initialize()
 		{
-			TrackComponents = FMovieSceneTracksComponentTypes::Get();
 		}
 
-		void ForEachAllocation(
-			const FEntityAllocation* Allocation,
-			TRead<FObjectKey> RootMotionTargets,
-			TReadOneOrMoreOf<
-			double, double, double,
-			double, double, double,
-			double, double, double>
-			TransformBlendResults) const
+#else
+
+		bool IsActive() const
 		{
-			using namespace UE::AnimNext;
-			using namespace UE::MovieScene;
+			return bIsActorBeingMoved.load(std::memory_order_relaxed);
+		}
 
-			const double* DoubleResults[9];
+		void Initialize()
+		{
+			static bool bFirst = true;
+			if (bFirst)
 			{
-				DoubleResults[0] = TransformBlendResults.Get<0>();
-				DoubleResults[1] = TransformBlendResults.Get<1>();
-				DoubleResults[2] = TransformBlendResults.Get<2>();
-				DoubleResults[3] = TransformBlendResults.Get<3>();
-				DoubleResults[4] = TransformBlendResults.Get<4>();
-				DoubleResults[5] = TransformBlendResults.Get<5>();
-				DoubleResults[6] = TransformBlendResults.Get<6>();
-				DoubleResults[7] = TransformBlendResults.Get<7>();
-				DoubleResults[8] = TransformBlendResults.Get<8>();
+				bFirst = false;
+				GEditor->OnBeginObjectMovement().AddRaw(this, &FIsActorBeingMovedStatics::Set, true);
+				GEditor->OnEndObjectMovement().AddRaw(this, &FIsActorBeingMovedStatics::Set, false);
 			}
+		}
 
-			const bool bFullyAnimated =
-				DoubleResults[0] && DoubleResults[1] && DoubleResults[2] &&
-				DoubleResults[3] && DoubleResults[4] && DoubleResults[5] &&
-				DoubleResults[6] && DoubleResults[7] && DoubleResults[8];
+	private:
 
+		void Set(UObject&, bool bIsActive)
+		{
+			ensure(IsInGameThread());
+			bIsActorBeingMoved.exchange(bIsActive, std::memory_order_relaxed);
+		}
 
-			TOptionalComponentReader<FIntermediate3DTransform> OptInitialValues =
-				Allocation->TryReadComponents(TrackComponents->Parameters.Transform.InitialValue);
+		mutable std::atomic<bool> bIsActorBeingMoved = false;
+#endif // WITH_EDITOR
 
-			for (int32 Index = 0; Index < Allocation->Num(); ++Index)
-			{
-				if (TSharedPtr<FMovieSceneMixerRootMotionComponentData>* RootMotionDataPtr = System->GetData().Find(RootMotionTargets[Index]))
-				{
-					// We have root motion for this object. Copy over the transform.
-					if (RootMotionDataPtr->IsValid())
-					{
-						if (bFullyAnimated)
-						{
-							(*RootMotionDataPtr)->SetSequenceAnimationTransform(FTransform(
-								FRotator(DoubleResults[4][Index], DoubleResults[5][Index], DoubleResults[3][Index]),
-								FVector(DoubleResults[0][Index], DoubleResults[1][Index], DoubleResults[2][Index]),
-								FVector(DoubleResults[6][Index], DoubleResults[7][Index], DoubleResults[8][Index])
-							));
-						}
-						else
-						{
-							FRotator Rotation = FRotator::ZeroRotator;
-							FVector Location = FVector::ZeroVector;
-							FVector Scale = FVector::OneVector;
-
-							if (OptInitialValues)
-							{
-								Rotation = OptInitialValues[Index].GetRotation();
-								Location = OptInitialValues[Index].GetTranslation();
-								Scale = OptInitialValues[Index].GetScale();
-							}
-
-							if (DoubleResults[0]) { Location.X = DoubleResults[0][Index]; }
-							if (DoubleResults[1]) { Location.Y = DoubleResults[1][Index]; }
-							if (DoubleResults[2]) { Location.Z = DoubleResults[2][Index]; }
-
-							if (DoubleResults[3]) { Rotation.Roll = DoubleResults[3][Index]; }
-							if (DoubleResults[4]) { Rotation.Pitch = DoubleResults[4][Index]; }
-							if (DoubleResults[5]) { Rotation.Yaw = DoubleResults[5][Index]; }
-
-							if (DoubleResults[6]) { Scale.X = DoubleResults[6][Index]; }
-							if (DoubleResults[7]) { Scale.Y = DoubleResults[7][Index]; }
-							if (DoubleResults[8]) { Scale.Z = DoubleResults[8][Index]; }
-
-							(*RootMotionDataPtr)->SetSequenceAnimationTransform(FTransform(Rotation, Location, Scale));
-						}
-					}
-				}
-			}
-		}
-	};
-};
+	} GActorMovementTracker;
 
+} // namespace UE::MovieScene
 
-TOptional<FTransform> FMovieSceneMixerRootMotionComponentData::GetSequenceAnimationTransform() const
-{
-	UE::TReadScopeLock ScopeLock(RootMotionLock);
-	return SequenceAnimationTransform;
-}
 
-void FMovieSceneMixerRootMotionComponentData::SetSequenceAnimationTransform(const FTransform& InTransform)
-{
-	UE::TWriteScopeLock ScopeLock(RootMotionLock);
-	SequenceAnimationTransform = InTransform;
-}
 
 TOptional<FQuat> FMovieSceneMixerRootMotionComponentData::GetInverseMeshToActorRotation() const
 {
@@ -173,44 +120,51 @@ TOptional<FQuat> FMovieSceneMixerRootMotionComponentData::GetInverseMeshToActorR
 	return InverseMeshToActorRotation;
 }
 
-void FMovieSceneMixerRootMotionComponentData::SetInverseMeshToActorRotation(const FQuat& InInverseMeshToActorRotation)
+void FMovieSceneMixerRootMotionComponentData::Initialize()
 {
 	UE::TWriteScopeLock ScopeLock(RootMotionLock);
-	InverseMeshToActorRotation = InInverseMeshToActorRotation;
-}
 
-TOptional<FTransform> FMovieSceneMixerRootMotionComponentData::GetRootTransform() const
-{
-	UE::TReadScopeLock ScopeLock(RootMotionLock);
-	return RootTransform;
-}
+	InverseMeshToActorRotation.Reset();
 
-bool FMovieSceneMixerRootMotionComponentData::ConsumeRootTransform(FTransform& OutTransform)
-{
-	UE::TWriteScopeLock ScopeLock(RootMotionLock);
-	if (RootTransform.IsSet())
+	USceneComponent* BoundObject = OriginalBoundObject.Get();
+	USceneComponent* NewTarget = nullptr;
+
+	if (!BoundObject)
 	{
-		OutTransform = RootTransform.GetValue();
-		RootTransform.Reset();
-		return true;
+		Target = nullptr;
+		// Leave the last known component and actor transform
+		return;
 	}
-	return false;
-}
 
+	// If we're applying to the actor or a root custom attribute, we need to factor out the actor->compont rotation
+	const bool bNeedInverseMeshRotation = RootDestination == EMovieSceneRootMotionDestination::Actor
+			|| RootDestination == EMovieSceneRootMotionDestination::Attribute;
 
-void FMovieSceneMixerRootMotionComponentData::SetRootTransform(const FTransform& InTransform)
-{
-	UE::TWriteScopeLock ScopeLock(RootMotionLock);
-	RootTransform = InTransform;
-}
+	USceneComponent* RootComponent = BoundObject->GetOwner()->GetRootComponent();
+	if (bNeedInverseMeshRotation && RootComponent && RootComponent != BoundObject)
+	{
+		InverseMeshToActorRotation = RootComponent->GetComponentTransform().GetRelativeTransformReverse(BoundObject->GetComponentTransform()).GetRotation();
+	}
 
+	ComponentTransform = BoundObject->GetRelativeTransform();
+	ActorTransform = RootComponent->GetRelativeTransform();
 
-void FMovieSceneMixerRootMotionComponentData::Reset()
-{
-	UE::TWriteScopeLock ScopeLock(RootMotionLock);
-	SequenceAnimationTransform.Reset();
-	InverseMeshToActorRotation.Reset();
-	RootTransform.Reset();
+	switch (RootDestination)
+	{
+	case EMovieSceneRootMotionDestination::Discard:
+	case EMovieSceneRootMotionDestination::RootBone:
+	case EMovieSceneRootMotionDestination::Attribute:
+		Target = nullptr;
+		break;
+
+	case EMovieSceneRootMotionDestination::Component:
+		Target = BoundObject;
+		break;
+
+	case EMovieSceneRootMotionDestination::Actor:
+		Target = RootComponent;
+		break;
+	}
 }
 
 FAnimNextConvertRootMotionToWorldSpaceTask::FAnimNextConvertRootMotionToWorldSpaceTask(const TSharedPtr<FMovieSceneMixerRootMotionComponentData>& InRootMotionData)
@@ -228,25 +182,38 @@ void FAnimNextConvertRootMotionToWorldSpaceTask::Execute(UE::AnimNext::FEvaluati
 	using namespace UE::AnimNext; 
 	using namespace UE::MovieScene;
 
+	if (!EnumHasAnyFlags(VM.GetFlags(), EEvaluationFlags::Attributes | EEvaluationFlags::Trajectory))
+	{
+		return;
+	}
+
+	TUniquePtr<FKeyframeState>* Keyframe = VM.PeekValueMutable<TUniquePtr<FKeyframeState>>(KEYFRAME_STACK_NAME, 0);
+	if (!Keyframe || !Keyframe->IsValid())
+	{
+		return;
+	}
+
+
 	FAnimMixerComponentTypes* AnimMixerComponents = FAnimMixerComponentTypes::Get();
-	TOptional<FTransform> SequenceAnimationTransform = RootMotionData.IsValid() ? RootMotionData->GetSequenceAnimationTransform() : TOptional<FTransform>();
-	if (SequenceAnimationTransform.IsSet() && EnumHasAnyFlags(VM.GetFlags(), EEvaluationFlags::Attributes | EEvaluationFlags::Trajectory))
+	if (RootMotionData)
 	{
-		if (TUniquePtr<FKeyframeState>* Keyframe = VM.PeekValueMutable<TUniquePtr<FKeyframeState>>(KEYFRAME_STACK_NAME, 0))
+		if (FTransformAnimationAttribute* RootMotionAttribute = (*Keyframe)->Attributes.Find<FTransformAnimationAttribute>(AnimMixerComponents->RootTransformAttributeId))
 		{
-			if (Keyframe->IsValid())
+			FTransform RelativeTransform = RootMotionAttribute->Value;
+
+			TOptional<FQuat> InverseMeshToActorRotation = RootMotionData->GetInverseMeshToActorRotation();
+			if (InverseMeshToActorRotation.IsSet())
 			{
-				if (FTransformAnimationAttribute* RootMotionAttribute = (*Keyframe)->Attributes.Find<FTransformAnimationAttribute>(AnimMixerComponents->RootTransformAttributeId))
-				{
-					FTransform RelativeTransform = RootMotionAttribute->Value;
-					TOptional<FQuat> InverseMeshToActorRotation = RootMotionData->GetInverseMeshToActorRotation();
-					if (InverseMeshToActorRotation.IsSet())
-					{
-						RelativeTransform.SetTranslation(InverseMeshToActorRotation.GetValue() * RelativeTransform.GetTranslation());
-					}
-					RelativeTransform = RelativeTransform * SequenceAnimationTransform.GetValue();
-					RootMotionAttribute->Value = RelativeTransform;
-				}
+				RelativeTransform.SetTranslation(InverseMeshToActorRotation.GetValue() * RelativeTransform.GetTranslation());
+			}
+
+			if (FTransformAnimationAttribute* RootBaseAttribute = (*Keyframe)->Attributes.Find<FTransformAnimationAttribute>(AnimMixerComponents->RootTransformBaseAttributeId))
+			{
+				RootMotionAttribute->Value = RelativeTransform * RootBaseAttribute->Value;
+			}
+			else
+			{
+				RootMotionAttribute->Value = RelativeTransform;
 			}
 		}
 	}
@@ -271,18 +238,149 @@ void FAnimNextStoreRootTransformTask::Execute(UE::AnimNext::FEvaluationVM& VM) c
 
 	FAnimMixerComponentTypes* AnimMixerComponents = FAnimMixerComponentTypes::Get();
 
-	if (RootMotionData.IsValid() && EnumHasAnyFlags(VM.GetFlags(), EEvaluationFlags::Attributes | EEvaluationFlags::Trajectory))
+	if (!RootMotionData.IsValid() || !EnumHasAnyFlags(VM.GetFlags(), EEvaluationFlags::Attributes | EEvaluationFlags::Trajectory))
+	{
+		return;
+	}
+
+	TUniquePtr<FKeyframeState>* Keyframe = VM.PeekValueMutable<TUniquePtr<FKeyframeState>>(KEYFRAME_STACK_NAME, 0);
+	if (!Keyframe || !Keyframe->IsValid())
+	{
+		return;
+	}
+
+	FTransformAnimationAttribute* RootMotionAttribute = (*Keyframe)->Attributes.Find<FTransformAnimationAttribute>(AnimMixerComponents->RootTransformAttributeId);
+	FTransformAnimationAttribute* RootBaseAttribute   = (*Keyframe)->Attributes.Find<FTransformAnimationAttribute>(AnimMixerComponents->RootTransformBaseAttributeId);
+	if (!RootMotionAttribute && !RootBaseAttribute)
+	{
+		return;
+	}
+
+	EMovieSceneRootMotionDestination FinalDestination = RootMotionData->RootDestination;
+
+	FIntegerAnimationAttribute* Int = (*Keyframe)->Attributes.Find<FIntegerAnimationAttribute>(AnimMixerComponents->LegacySwapRootBoneAttributeId);
+	if (Int)
+	{
+		if (Int->Value == (int32)ESwapRootBone::SwapRootBone_Actor)
+		{
+			FinalDestination = EMovieSceneRootMotionDestination::Actor;
+		}
+		else if (Int->Value == (int32)ESwapRootBone::SwapRootBone_Component)
+		{
+			FinalDestination = EMovieSceneRootMotionDestination::Component;
+		}
+	}
+
+	FTransform RootTransform = RootMotionAttribute
+		? RootMotionAttribute->Value
+		: RootBaseAttribute->Value;
+
+	bool bRemoveRootMotionAttributes = true;
+
+	const int32 RootBoneIndex = (*Keyframe)->Pose.GetRefPose().GetLODBoneIndexFromSkeletonBoneIndex(0);
+
+	switch (RootMotionData->RootDestination)
+	{
+	case EMovieSceneRootMotionDestination::Discard:
+		if (RootBoneIndex != INDEX_NONE)
+		{
+			(*Keyframe)->Pose.LocalTransformsView[RootBoneIndex] = FTransform::Identity;
+		}
+		return;
+
+	case EMovieSceneRootMotionDestination::RootBone:
+		if (RootBoneIndex != INDEX_NONE)
+		{
+			(*Keyframe)->Pose.LocalTransformsView[RootBoneIndex] = RootTransform;
+		}
+		return;
+
+	case EMovieSceneRootMotionDestination::Component:
+		break;
+
+	case EMovieSceneRootMotionDestination::Actor:
+		break;
+
+	case EMovieSceneRootMotionDestination::Attribute:
+		bRemoveRootMotionAttributes = false;
+		break;
+	}
+
+
+	const FFloatAnimationAttribute* RootMotionWeight = (*Keyframe)->Attributes.Find<FFloatAnimationAttribute>(AnimMixerComponents->RootTransformWeightAttributeId);
+	const float TransformWeight = RootMotionWeight ? RootMotionWeight->Value : 1.f;
+
+	if (TransformWeight < 1.f)
 	{
-		if (TUniquePtr<FKeyframeState>* Keyframe = VM.PeekValueMutable<TUniquePtr<FKeyframeState>>(KEYFRAME_STACK_NAME, 0))
+		static const UE::Anim::FAttributeId RootMotionDeltaAttributeId(UE::Anim::IAnimRootMotionProvider::AttributeName , FCompactPoseBoneIndex(0));
+
+		FTransformAnimationAttribute* RootMotionDeltaAttribute = (*Keyframe)->Attributes.Find<FTransformAnimationAttribute>(RootMotionDeltaAttributeId);
+
+		// Undo the weight of the root motion
+		RootTransform *= ScalarRegister(1.f / TransformWeight);
+		RootTransform.NormalizeRotation();
+
+		FTransform RootMotionDelta;
+		if (RootMotionDeltaAttribute)
 		{
-			if (Keyframe->IsValid())
+			// Factor in the current root motion delta in order to support blending to/from gameplay-root-motion
+			RootMotionDelta = RootMotionDeltaAttribute->Value;
+
+			TOptional<FQuat> InverseMeshToActorRotation = RootMotionData->GetInverseMeshToActorRotation();
+			if (InverseMeshToActorRotation.IsSet())
 			{
-				if (FTransformAnimationAttribute* RootMotionAttribute = (*Keyframe)->Attributes.Find<FTransformAnimationAttribute>(AnimMixerComponents->RootTransformAttributeId))
+				RootMotionDelta.SetTranslation(InverseMeshToActorRotation.GetValue() * RootMotionDelta.GetTranslation());
+			}
+		}
+
+		FTransform ActualTransform = (RootMotionData->RootDestination == EMovieSceneRootMotionDestination::Component)
+			? RootMotionData->ComponentTransform + RootMotionDelta
+			: RootMotionData->ActorTransform + RootMotionDelta;
+
+		ActualTransform.NormalizeRotation();
+		ActualTransform.BlendWith(RootTransform, TransformWeight);
+
+		RootTransform = ActualTransform;
+
+		if (RootMotionAttribute)
+		{
+			RootMotionAttribute->Value = RootTransform;
+		}
+		else if (!bRemoveRootMotionAttributes)
+		{
+			(*Keyframe)->Attributes.Add(AnimMixerComponents->RootTransformAttributeId, FTransformAnimationAttribute(RootTransform));
+		}
+	}
+
+	if (RootMotionData->RootDestination != EMovieSceneRootMotionDestination::Attribute)
+	{
+		TWeakPtr<FMovieSceneMixerRootMotionComponentData> WeakRootMotion = RootMotionData;
+
+		auto ApplyRootMotion = [RootTransform, Destination = RootMotionData->RootDestination, WeakTarget = RootMotionData->Target]
+		{
+			USceneComponent* TargetComponent = WeakTarget.Get();
+			if (TargetComponent && !GActorMovementTracker.IsActive())
+			{
+				if (Destination == EMovieSceneRootMotionDestination::Actor)
+				{
+					TargetComponent = TargetComponent->GetOwner()->GetRootComponent();
+				}
+
+				if (TargetComponent)
 				{
-					RootMotionData->SetRootTransform(RootMotionAttribute->Value);
+					TargetComponent->SetRelativeLocationAndRotation(RootTransform.GetLocation(), RootTransform.GetRotation().Rotator());
 				}
 			}
-		}
+		};
+		FAnimNextModuleInstance::RunTaskOnGameThread(MoveTemp(ApplyRootMotion));
+	}
+
+
+	// We're done with these attributes now
+	if (bRemoveRootMotionAttributes)
+	{
+		(*Keyframe)->Attributes.Remove<FTransformAnimationAttribute>(AnimMixerComponents->RootTransformAttributeId);
+		(*Keyframe)->Attributes.Remove<FFloatAnimationAttribute>(AnimMixerComponents->RootTransformWeightAttributeId);
 	}
 }
 
@@ -300,47 +398,23 @@ UMovieSceneRootMotionSystem::UMovieSceneRootMotionSystem(const FObjectInitialize
 
 	if (HasAnyFlags(RF_ClassDefaultObject))
 	{
+		// This must be run before the anim mixer to ensure that the anim mixer sets up its program correctly with root motion
+		DefineImplicitPrerequisite(GetClass(), UMovieSceneAnimMixerSystem::StaticClass());
+
+		DefineImplicitPrerequisite(UDoubleChannelEvaluatorSystem::StaticClass(), GetClass());
 		DefineImplicitPrerequisite(UMovieScenePiecewiseDoubleBlenderSystem::StaticClass(), GetClass());
-		DefineImplicitPrerequisite(UMovieSceneAnimMixerSystem::StaticClass(), GetClass()); 
+		DefineImplicitPrerequisite(UMovieSceneComponentTransformSystem::StaticClass(), GetClass());
 	}
 }
 
-#if WITH_EDITOR
 void UMovieSceneRootMotionSystem::OnLink()
 {
-	if (OnBeginActorMovementHandle.IsValid() == false)
-	{
-		OnBeginActorMovementHandle = GEditor->OnBeginObjectMovement().Add(UEditorEngine::FOnBeginTransformObject::FDelegate::CreateLambda([this](UObject&)
-			{
-				bActorBeingMoved = true;
-			}));
-	}
-
-	if (OnEndActorMovementHandle.IsValid() == false)
-	{
-		OnEndActorMovementHandle = GEditor->OnEndObjectMovement().Add(UEditorEngine::FOnEndTransformObject::FDelegate::CreateLambda([this](UObject&)
-			{
-				bActorBeingMoved = false;
-			}));
-	}
-}
-
-void UMovieSceneRootMotionSystem::OnUnlink()
-{
-	GEditor->OnBeginObjectMovement().Remove(OnBeginActorMovementHandle);
-	OnBeginActorMovementHandle.Reset();
-	GEditor->OnEndObjectMovement().Remove(OnEndActorMovementHandle);
-	OnEndActorMovementHandle.Reset();
-}
-
-#endif
+	UMovieSceneAnimMixerSystem* AnimMixer = Linker->LinkSystem<UMovieSceneAnimMixerSystem>();
+	Linker->SystemGraph.AddReference(AnimMixer, this);
 
-void UMovieSceneRootMotionSystem::FinishDestroy()
-{
-	Super::FinishDestroy();
+	UE::MovieScene::GActorMovementTracker.Initialize();
 }
 
-
 void UMovieSceneRootMotionSystem::OnRun(FSystemTaskPrerequisites& InPrerequisites, FSystemSubsequentTasks& Subsequents)
 {
 	using namespace UE::MovieScene;
@@ -352,58 +426,57 @@ void UMovieSceneRootMotionSystem::OnRun(FSystemTaskPrerequisites& InPrerequisite
 
 	if (Runner->GetCurrentPhase() == ESystemPhase::Instantiation)
 	{
-		TArray<FMovieSceneEntityID> StaleRootMotion;
-		// Remove stale root motion from expiring mixers
-		FEntityTaskBuilder()
-			.ReadEntityIDs()
-			.FilterAll({ AnimMixerComponents->MixerTask, BuiltInComponents->Tags.NeedsUnlink })
-			.Iterate_PerEntity(&Linker->EntityManager, [&StaleRootMotion](FMovieSceneEntityID EntityID)
-				{
-					StaleRootMotion.Add(EntityID);
-				});
+		UMovieSceneAnimMixerSystem* AnimMixer = Linker->FindSystem<UMovieSceneAnimMixerSystem>();
+		if (!ensure(AnimMixer))
+		{
+			// Should never exist without the anim mixer
+			return;
+		}
 
-		// Remove stale root motion from mixers without root motion components that need linking
+		// Remove expiring root motions
 		FEntityTaskBuilder()
-			.ReadEntityIDs()
-			.FilterNone({ AnimMixerComponents->MixerRootMotion })
-			.FilterAll({ BuiltInComponents->Tags.NeedsLink, AnimMixerComponents->MixerTask})
-			.Iterate_PerEntity(&Linker->EntityManager, [&StaleRootMotion](FMovieSceneEntityID EntityID)
-				{
-					StaleRootMotion.Add(EntityID);
-				});
-
-		if (!StaleRootMotion.IsEmpty())
-		{
-			for (auto It = RootMotionData.CreateIterator(); It; ++It)
+		.Write(AnimMixerComponents->MixerRootMotion)
+		.FilterAll({ AnimMixerComponents->RootDestination, BuiltInComponents->Tags.NeedsUnlink })
+		.Iterate_PerEntity( &Linker->EntityManager,
+			[](TSharedPtr<FMovieSceneMixerRootMotionComponentData>& OutRootMotion)
 			{
-				if (It->Value.IsValid() && StaleRootMotion.Contains(It->Value->MixerEntityID))
-				{
-					It.RemoveCurrent();
-				}
+				OutRootMotion = nullptr;
 			}
-		}
+		);
 
-		// Track or update root motion
+		// Set up root motion behaviors
 		FEntityTaskBuilder()
-		.ReadEntityIDs()
-		.Read(AnimMixerComponents->MixerRootMotion)
-		.FilterAll({BuiltInComponents->Tags.NeedsLink })
-		.Iterate_PerEntity(&Linker->EntityManager, [this](FMovieSceneEntityID EntityID, TSharedPtr<FMovieSceneMixerRootMotionComponentData> InRootMotion)
+		.Read(BuiltInComponents->BoundObjectKey)
+		.Write(AnimMixerComponents->Target)
+		.Write(AnimMixerComponents->MixerRootMotion)
+		.FilterAll({ AnimMixerComponents->RootDestination, BuiltInComponents->Tags.NeedsLink })
+		.Iterate_PerEntity( &Linker->EntityManager,
+			[AnimMixer](FObjectKey BoundObjectKey, TInstancedStruct<FMovieSceneMixedAnimationTarget>& OutTarget, TSharedPtr<FMovieSceneMixerRootMotionComponentData>& OutRootMotion)
 			{
-				// Clear any existing data on that entity
-				for (auto It = RootMotionData.CreateIterator(); It; ++It)
+				// For new entities, we 'resolve' the animation target so if 'Automatic' is picked we choose the right one automatically.
+				OutTarget = UMovieSceneAnimMixerSystem::ResolveAnimationTarget(BoundObjectKey, OutTarget);
+
+				FMovieSceneAnimMixerKey Key(BoundObjectKey, OutTarget);
+
+				if (OutRootMotion)
 				{
-					if (It->Value.IsValid() && It->Value->MixerEntityID == EntityID)
-					{
-						It.RemoveCurrent();
-					}
+					AnimMixer->AssignRootMotion(Key, OutRootMotion);
 				}
-
-				if (InRootMotion.IsValid() && InRootMotion->Target.IsValid())
+				else if (TSharedPtr<FMovieSceneMixerRootMotionComponentData> RootMotion = AnimMixer->FindRootMotion(Key))
+				{
+					OutRootMotion = RootMotion;
+				}
+				else
 				{
-					RootMotionData.Add(InRootMotion->Target.Get(), InRootMotion);
+					// Make a new one
+					OutRootMotion = MakeShared<FMovieSceneMixerRootMotionComponentData>();
+					AnimMixer->AssignRootMotion(Key, OutRootMotion);
 				}
-			});
+
+				USceneComponent* Target = Cast<USceneComponent>(BoundObjectKey.ResolveObjectPtr());
+				OutRootMotion->OriginalBoundObject = Target;
+			}
+		);
 	}
 }
 
@@ -415,58 +488,30 @@ void UMovieSceneRootMotionSystem::OnSchedulePersistentTasks(UE::MovieScene::IEnt
 	FMovieSceneTracksComponentTypes* TrackComponents = FMovieSceneTracksComponentTypes::Get();
 	FAnimMixerComponentTypes* AnimMixerComponents = FAnimMixerComponentTypes::Get();
 
-	// Zero out all transforms to start
-	for (TTuple<FObjectKey, TSharedPtr<FMovieSceneMixerRootMotionComponentData>>& RootMotion : RootMotionData)
-	{
-		RootMotion.Value->Reset();
-	}
-
-	// Update some optional root motion data used later in the frame for calculation
-	FEntityTaskBuilder()
-		.Read(AnimMixerComponents->MeshComponent)
-		.Write(AnimMixerComponents->MixerRootMotion)
-		.SetDesiredThread(Linker->EntityManager.GetGatherThread())
-		.Schedule_PerAllocation<FUpdateRootMotionData>(&Linker->EntityManager, TaskScheduler);
-
-	// Copy the component transform blend result to the root motion component data
-	FEntityTaskBuilder()
-	.Read(BuiltInComponents->BoundObjectKey)
-	.ReadOneOrMoreOf(BuiltInComponents->DoubleResult[0],
-					BuiltInComponents->DoubleResult[1],
-					BuiltInComponents->DoubleResult[2],
-					BuiltInComponents->DoubleResult[3],
-					BuiltInComponents->DoubleResult[4],
-					BuiltInComponents->DoubleResult[5],
-					BuiltInComponents->DoubleResult[6],
-					BuiltInComponents->DoubleResult[7],
-					BuiltInComponents->DoubleResult[8])
-	.FilterAll({TrackComponents->ComponentTransform.PropertyTag})
-	.FilterNone({BuiltInComponents->BlendChannelInput})
-	.SetDesiredThread(Linker->EntityManager.GetGatherThread())
-	.Schedule_PerAllocation<FSyncTransformTrackResult>(&Linker->EntityManager, TaskScheduler, this);
-}
-
-void UMovieSceneRootMotionSystem::ApplyRootMotion()
-{
-#if WITH_EDITOR
-	if (bActorBeingMoved)
+	UMovieSceneAnimMixerSystem* AnimMixer = Linker->FindSystem<UMovieSceneAnimMixerSystem>();
+	if (!ensure(AnimMixer))
 	{
 		return;
 	}
-#endif
 
-	for (TTuple<FObjectKey, TSharedPtr<FMovieSceneMixerRootMotionComponentData>> RootMotionTuple : RootMotionData)
-	{
-		if (TSharedPtr<FMovieSceneMixerRootMotionComponentData> RootMotion = RootMotionTuple.Value)
-		{
-			if (RootMotion->Target.IsValid())
-			{
-				FTransform RelativeTransform;
-				if (RootMotion->ConsumeRootTransform(RelativeTransform))
-				{
-					RootMotion->Target->SetRelativeLocationAndRotation(RelativeTransform.GetLocation(), RelativeTransform.GetRotation().Rotator());
-				}
-			}
-		}
-	}
+	FTaskID WaitForAllTransforms = FEntityTaskBuilder()
+	.Write(BuiltInComponents->CustomPropertyIndex)
+	.FilterAll({ TrackComponents->ComponentTransform.PropertyTag })
+	.Fork_PerAllocation<FNoopTask>(&Linker->EntityManager, TaskScheduler);
+
+	// Gather root destination results
+	FTaskID GatherRootDestinationTask = FEntityTaskBuilder()
+	.Read(BuiltInComponents->ByteResult)
+	.Write(AnimMixerComponents->RootDestination)
+	.Write(AnimMixerComponents->MixerRootMotion)
+	.FilterNone({ BuiltInComponents->Tags.Ignored })
+	.Schedule_PerEntity<FGatherRootDestinations>(&Linker->EntityManager, TaskScheduler);
+
+	// Reset root motion data and gather current component/actor transforms
+	FTaskID InitializeRootMotion = TaskScheduler->AddTask<FInitializeRootMotionTask>(
+		FTaskParams(TEXT("Initialize Root Motion")).ForceGameThread(),
+		AnimMixer);
+
+	TaskScheduler->AddPrerequisite(WaitForAllTransforms,      InitializeRootMotion);
+	TaskScheduler->AddPrerequisite(GatherRootDestinationTask, InitializeRootMotion);
 }
diff --git a/Engine/Plugins/Experimental/MovieSceneAnimMixer/Source/MovieSceneAnimMixer/Public/AnimMixerComponentTypes.h b/Engine/Plugins/Experimental/MovieSceneAnimMixer/Source/MovieSceneAnimMixer/Public/AnimMixerComponentTypes.h
--- a/Engine/Plugins/Experimental/MovieSceneAnimMixer/Source/MovieSceneAnimMixer/Public/AnimMixerComponentTypes.h
+++ b/Engine/Plugins/Experimental/MovieSceneAnimMixer/Source/MovieSceneAnimMixer/Public/AnimMixerComponentTypes.h
@@ -13,6 +13,10 @@
 #include "AnimSequencerInstanceProxy.h"
 #include "Systems/MovieSceneRootMotionSystem.h"
 
+enum class EMovieSceneRootMotionDestination : uint8;
+
+struct FMovieSceneAnimMixerEntry;
+
 namespace UE::MovieScene
 {
 	struct FAnimMixerComponentTypes
@@ -22,12 +26,17 @@ namespace UE::MovieScene
 
 		// A root motion transform in animation space to be applied (after space conversion)
 		MOVIESCENEANIMMIXER_API static const UE::Anim::FAttributeId RootTransformAttributeId;
+		MOVIESCENEANIMMIXER_API static const UE::Anim::FAttributeId RootTransformBaseAttributeId;
+		MOVIESCENEANIMMIXER_API static const UE::Anim::FAttributeId RootTransformWeightAttributeId;
+		MOVIESCENEANIMMIXER_API static const UE::Anim::FAttributeId LegacySwapRootBoneAttributeId;
 
 		TComponentTypeID<int32> Priority;
 		TComponentTypeID<TInstancedStruct<FMovieSceneMixedAnimationTarget>> Target;
 		TComponentTypeID<TSharedPtr<FAnimNextEvaluationTask>> Task;
 		TComponentTypeID<TSharedPtr<FAnimNextEvaluationTask>> MixerTask;
+		TComponentTypeID<TSharedPtr<FMovieSceneAnimMixerEntry>> MixerEntry;
 		TComponentTypeID<FMovieSceneRootMotionSettings> RootMotionSettings;
+		TComponentTypeID<EMovieSceneRootMotionDestination> RootDestination;
 		TComponentTypeID<FObjectComponent> MeshComponent;
 		TComponentTypeID<TSharedPtr<FMovieSceneMixerRootMotionComponentData>> MixerRootMotion;
 
diff --git a/Engine/Plugins/Experimental/MovieSceneAnimMixer/Source/MovieSceneAnimMixer/Public/Systems/MovieSceneAnimInstanceTargetSystem.h b/Engine/Plugins/Experimental/MovieSceneAnimMixer/Source/MovieSceneAnimMixer/Public/Systems/MovieSceneAnimInstanceTargetSystem.h
--- a/Engine/Plugins/Experimental/MovieSceneAnimMixer/Source/MovieSceneAnimMixer/Public/Systems/MovieSceneAnimInstanceTargetSystem.h
+++ b/Engine/Plugins/Experimental/MovieSceneAnimMixer/Source/MovieSceneAnimMixer/Public/Systems/MovieSceneAnimInstanceTargetSystem.h
@@ -13,7 +13,6 @@
 
 
 struct FAnimNextEvaluationTask;
-struct FMovieSceneMixerRootMotionComponentData;
 
 /**
  * Declaring a unique target for targeting a custom anim instance on the skeletal mesh component. Doesn't need additional metadata.
@@ -40,7 +39,6 @@ class USequencerMixedAnimInstance : public UAnimInstance, public ISequencerAnima
 
 public:
 	MOVIESCENEANIMMIXER_API void SetMixerTask(TSharedPtr<FAnimNextEvaluationTask> InEvalTask);
-	MOVIESCENEANIMMIXER_API void SetRootMotion(const TSharedPtr<FMovieSceneMixerRootMotionComponentData>& InRootMotion);
 
 protected:
 	// UAnimInstance interface
@@ -92,7 +90,6 @@ public:
 	UAnimInstance* GetSourceAnimInstance() const { return CurrentSourceAnimInstance; }
 
 	void SetMixerTask(TSharedPtr<FAnimNextEvaluationTask> InEvalTask);
-	void SetRootMotion(const TSharedPtr<FMovieSceneMixerRootMotionComponentData>& InRootMotion);
 	void LinkSourcePose(UAnimInstance* InInputInstance, FAnimInstanceProxy* InInputProxy);
 	void UnlinkSourcePose();
 
@@ -107,9 +104,6 @@ protected:
 
 	// Pointer to the task passed from the mixer.
 	TSharedPtr<FAnimNextEvaluationTask> MixerTask;
-
-	// Root motion settings from the mixed animation
-	TSharedPtr<FMovieSceneMixerRootMotionComponentData> RootMotion;
 };
 
 
diff --git a/Engine/Plugins/Experimental/MovieSceneAnimMixer/Source/MovieSceneAnimMixer/Public/Systems/MovieSceneAnimMixerSystem.h b/Engine/Plugins/Experimental/MovieSceneAnimMixer/Source/MovieSceneAnimMixer/Public/Systems/MovieSceneAnimMixerSystem.h
--- a/Engine/Plugins/Experimental/MovieSceneAnimMixer/Source/MovieSceneAnimMixer/Public/Systems/MovieSceneAnimMixerSystem.h
+++ b/Engine/Plugins/Experimental/MovieSceneAnimMixer/Source/MovieSceneAnimMixer/Public/Systems/MovieSceneAnimMixerSystem.h
@@ -22,15 +22,25 @@ namespace UE
 	}
 }
 
+struct FMovieSceneAnimMixer;
+
 struct FMovieSceneAnimMixerEntry
 {
-	UE::MovieScene::FMovieSceneEntityID EntityID;
+	
+	TWeakPtr<FMovieSceneAnimMixer> WeakMixer;
 	TSharedPtr<FAnimNextEvaluationTask> EvalTask;
+	/** Shared pointer to the root motion for this entry if it came from a FMovieSceneRootMotionSettings on a mixer
+	 * @note: Only to be used for lifetime management to keep FMovieSceneAnimMixer::WeakRootMotion alive!
+	 *        This may be null even if the result of the mix still has root motion */
+	TSharedPtr<FMovieSceneMixerRootMotionComponentData> RootMotionLifetimeReference;
+	TOptional<FMovieSceneRootMotionSettings> RootMotionSettings;
 	double PoseWeight = 1.0;
 	int32 Priority = 0;
 	bool bAdditive = false;
 	bool bRequiresBlend = true;
-	TOptional<FMovieSceneRootMotionSettings> RootMotionSettings;
+	UE::MovieScene::FMovieSceneEntityID EntityID;
+
+
 	FORCEINLINE bool operator<(const FMovieSceneAnimMixerEntry& RHS) const
 	{
 		if (Priority == RHS.Priority)
@@ -51,7 +61,7 @@ struct FMovieSceneAnimMixerEntry
 struct FMovieSceneAnimMixerKey
 {
 	FObjectKey BoundObjectKey;
-	TInstancedStruct<FMovieSceneMixedAnimationTarget> Target;		
+	TInstancedStruct<FMovieSceneMixedAnimationTarget> Target;
 
 	/** Equality operator */
 	friend bool operator==(const FMovieSceneAnimMixerKey& A, const FMovieSceneAnimMixerKey& B)
@@ -71,9 +81,11 @@ struct FMovieSceneAnimMixerKey
 struct FMovieSceneAnimMixer
 {
 	UE::MovieScene::FMovieSceneEntityID MixerEntityID;
-	TArray<FMovieSceneAnimMixerEntry> Entries;
+	TArray<TWeakPtr<FMovieSceneAnimMixerEntry>> WeakEntries;
 	TSharedPtr<UE::AnimNext::FEvaluationProgram> EvaluationProgram;
-	TSharedPtr<FMovieSceneMixerRootMotionComponentData> RootMotion;
+	TWeakPtr<FMovieSceneMixerRootMotionComponentData> WeakRootMotion;
+
+	bool bNeedsResort = false;
 };
 
 
@@ -119,17 +131,24 @@ public:
 
 	UMovieSceneAnimMixerSystem(const FObjectInitializer& ObjInit);
 
-	// Map of animation mixers
-	TMap<FMovieSceneAnimMixerKey, FMovieSceneAnimMixer> Mixers;
+	static TInstancedStruct<FMovieSceneMixedAnimationTarget> ResolveAnimationTarget(FObjectKey ObjectKey, const TInstancedStruct<FMovieSceneMixedAnimationTarget>& InTarget);
 
+	TSharedPtr<FMovieSceneMixerRootMotionComponentData> FindRootMotion(FMovieSceneAnimMixerKey InMixerKey) const;
+	void AssignRootMotion(FMovieSceneAnimMixerKey InMixerKey, TSharedPtr<FMovieSceneMixerRootMotionComponentData> RootMotion);
+
+	void InitializeAllRootMotion();
 
 private:
 
+	// Map of animation mixers
+	TMap<FMovieSceneAnimMixerKey, TSharedPtr<FMovieSceneAnimMixer>> Mixers;
+
+	virtual void OnLink() override;
+	virtual bool IsRelevantImpl(UMovieSceneEntitySystemLinker* InLinker) const override;
 	virtual void OnRun(FSystemTaskPrerequisites& InPrerequisites, FSystemSubsequentTasks& Subsequents) override;
 	virtual void OnSchedulePersistentTasks(UE::MovieScene::IEntitySystemScheduler* TaskScheduler) override;
-	virtual void OnCleanTaggedGarbage() override; 
-	
-	const TInstancedStruct<FMovieSceneMixedAnimationTarget> ResolveAnimationTarget(FObjectKey ObjectKey, const TInstancedStruct<FMovieSceneMixedAnimationTarget> InTarget);
+	virtual void OnCleanTaggedGarbage() override;
 
-	
+	UPROPERTY()
+	TObjectPtr<UMovieSceneRootMotionSystem> RootMotionSystem;
 };
\ No newline at end of file
