Sequencer: Initial Animation Mixer Track prototype

diff --git a/Engine/Plugins/Experimental/MovieSceneAnimMixer/Source/MovieSceneAnimMixer/Public/Systems/MovieSceneMixedSkeletalAnimationSystem.h b/Engine/Plugins/Experimental/MovieSceneAnimMixer/Source/MovieSceneAnimMixer/Public/Systems/MovieSceneMixedSkeletalAnimationSystem.h
--- a/Engine/Plugins/Experimental/MovieSceneAnimMixer/Source/MovieSceneAnimMixer/Public/Systems/MovieSceneMixedSkeletalAnimationSystem.h
+++ b/Engine/Plugins/Experimental/MovieSceneAnimMixer/Source/MovieSceneAnimMixer/Public/Systems/MovieSceneMixedSkeletalAnimationSystem.h
@@ -34,10 +34,10 @@ struct FMixedAnimSkeletalAnimationData
 	double ToPosition = -1.0f;
 
 	UPROPERTY()
-	bool bFireNotifies = true;
+	ESwapRootBone LegacySwapRootBone = ESwapRootBone::SwapRootBone_None;
 
 	UPROPERTY()
-	ESwapRootBone SwapRootBone = ESwapRootBone::SwapRootBone_None;
+	bool bFireNotifies = true;
 
 	UPROPERTY()
 	bool bAdditive = false;
@@ -55,6 +55,9 @@ struct FMovieSceneSkeletalAnimationEvaluationTask : public FAnimNextEvaluationTa
 
 	UPROPERTY()
 	FMixedAnimSkeletalAnimationData AnimationData;
+
+	UPROPERTY()
+	FTransform RootBaseTransform;
 };
 
 // System to handle creating evaluation tasks from skeletal animation track sections for the anim mixer.
diff --git a/Engine/Plugins/Experimental/MovieSceneAnimMixer/Source/MovieSceneAnimMixer/Public/Systems/MovieSceneRootMotionSystem.h b/Engine/Plugins/Experimental/MovieSceneAnimMixer/Source/MovieSceneAnimMixer/Public/Systems/MovieSceneRootMotionSystem.h
--- a/Engine/Plugins/Experimental/MovieSceneAnimMixer/Source/MovieSceneAnimMixer/Public/Systems/MovieSceneRootMotionSystem.h
+++ b/Engine/Plugins/Experimental/MovieSceneAnimMixer/Source/MovieSceneAnimMixer/Public/Systems/MovieSceneRootMotionSystem.h
@@ -9,6 +9,11 @@
 #include "Tickable.h"
 #include "MovieSceneRootMotionSystem.generated.h"
 
+enum class EMovieSceneRootMotionDestination : uint8;
+
+class UMovieSceneAnimMixerSystem;
+class UMovieSceneRootMotionSystem;
+
 // Enum describing what space that root motion should be applied in.
 UENUM()
 enum class EMovieSceneRootMotionSpace : uint8
@@ -23,31 +28,30 @@ USTRUCT()
 struct FSkeletalAnimationRootMotionOverride
 {
 	GENERATED_BODY()
-
-	/** If true we use the ChildBoneIndex otherwise we use the root*/
-	UPROPERTY()
-	bool bBlendFirstChildOfRoot = false;
-	
-	UPROPERTY()
-	int32 ChildBoneIndex = INDEX_NONE;
 	
 	UPROPERTY()
 	FTransform RootMotion = FTransform::Identity;
 	
 	UPROPERTY()
 	FTransform PreviousTransform = FTransform::Identity;
+	
+	UPROPERTY()
+	int32 ChildBoneIndex = INDEX_NONE;
+
+	/** If true we use the ChildBoneIndex otherwise we use the root*/
+	UPROPERTY()
+	bool bBlendFirstChildOfRoot = false;
 };
 
 // Structure used for animation tracks to communicate to the mixer how they would like root motion handled if at all.
 struct FMovieSceneRootMotionSettings
 {
-	// Whether to apply root motion to the component or the actor.
-	ESwapRootBone SwapRootBone = ESwapRootBone::SwapRootBone_None;
+	TOptional<FSkeletalAnimationRootMotionOverride> RootMotionOverride;
 
 	// What space to apply root motion in. Defaults to animation space.
 	EMovieSceneRootMotionSpace RootMotionSpace = EMovieSceneRootMotionSpace::AnimationSpace;
 
-	TOptional<FSkeletalAnimationRootMotionOverride> RootMotionOverride;
+	ESwapRootBone LegacySwapRootBone = ESwapRootBone::SwapRootBone_None;
 };
 
 // Structure that is shared between entities for handling the mixer's root motion result. 
@@ -59,20 +63,23 @@ public:
 	void SetSequenceAnimationTransform(const FTransform& InTransform);
 
 	TOptional<FQuat> GetInverseMeshToActorRotation() const;
-	void SetInverseMeshToActorRotation(const FQuat& InInverseMeshToActorRotation);
 
-	TOptional<FTransform> GetRootTransform() const;
-	bool ConsumeRootTransform(FTransform& OutTransform);
-	void SetRootTransform(const FTransform& InTransform);
-
-	void Reset();
+	void Initialize();
 
 public:
-		// Where to apply the root motion
-		TWeakObjectPtr<USceneComponent> Target;
 
-		// EntityID for the anim mixer
-		UE::MovieScene::FMovieSceneEntityID MixerEntityID;
+	TWeakObjectPtr<USceneComponent> OriginalBoundObject;
+
+	// Where to apply the root motion
+	TWeakObjectPtr<USceneComponent> Target;
+
+	// EntityID for the anim mixer
+	UE::MovieScene::FMovieSceneEntityID MixerEntityID;
+
+	EMovieSceneRootMotionDestination RootDestination;
+
+	FTransform ActorTransform;
+	FTransform ComponentTransform;
 
 private:
 	// Blend output result of transform track calculations
@@ -80,9 +87,6 @@ private:
 
 	// Optional inverse mesh component to actor rotation used to offset any mesh component rotation where applicable.
 	TOptional<FQuat> InverseMeshToActorRotation;
-	
-	// Final mixed root transform
-	TOptional<FTransform> RootTransform;
 
 	mutable FTransactionallySafeRWLock RootMotionLock;
 };
@@ -125,6 +129,7 @@ struct FAnimNextStoreRootTransformTask : public FAnimNextEvaluationTask
 	TSharedPtr<FMovieSceneMixerRootMotionComponentData> RootMotionData;
 };
 
+
 // Takes in evaluation tasks on mixers.
 // Mixes just the root motion attributes.
 // Converts it from animation space to either additive actor or component space (based on which attribute used).
@@ -132,7 +137,6 @@ struct FAnimNextStoreRootTransformTask : public FAnimNextEvaluationTask
 UCLASS(MinimalAPI)
 class UMovieSceneRootMotionSystem
 	: public UMovieSceneEntitySystem
-	, public FTickableGameObject
 {
 public:
 
@@ -140,40 +144,9 @@ public:
 
 	UMovieSceneRootMotionSystem(const FObjectInitializer& ObjInit);
 
-	TMap<FObjectKey, TSharedPtr<FMovieSceneMixerRootMotionComponentData>>& GetData() { return RootMotionData; }
-
 private:
 
+	virtual void OnLink() override;
 	virtual void OnRun(FSystemTaskPrerequisites& InPrerequisites, FSystemSubsequentTasks& Subsequents) override;
 	virtual void OnSchedulePersistentTasks(UE::MovieScene::IEntitySystemScheduler* TaskScheduler) override;
-	void FinishDestroy() override;
-
-	/** FTickableGameObject interface */
-	virtual UWorld* GetTickableGameObjectWorld() const override { return GetWorld(); }
-	virtual bool IsTickableInEditor() const override { return true; }
-	virtual bool IsTickable() const override { return !RootMotionData.IsEmpty(); }
-	virtual TStatId GetStatId() const override { RETURN_QUICK_DECLARE_CYCLE_STAT(UMovieSceneRootMotionSystem, STATGROUP_Tickables); }
-	virtual void Tick(float DeltaSeconds) override { ApplyRootMotion(); }
-
-	// Called by the Tick function to apply root motion after animation evaluation on the game thread
-	void ApplyRootMotion();
-
-	TMap<FObjectKey, TSharedPtr<FMovieSceneMixerRootMotionComponentData>> RootMotionData;
-
-#if WITH_EDITOR
-	void OnLink() override;
-	void OnUnlink() override;
-#endif
-
-#if WITH_EDITORONLY_DATA
-	FDelegateHandle OnBeginActorMovementHandle;
-	FDelegateHandle OnEndActorMovementHandle;
-	bool bActorBeingMoved = false;
-#endif
-
-	friend class FMovieSceneRootMotionHelper;
-
-
-protected:
-
-};
\ No newline at end of file
+};
diff --git a/Engine/Plugins/Experimental/MovieSceneAnimMixer/Source/MovieSceneAnimMixerEditor/MovieSceneAnimMixerEditor.Build.cs b/Engine/Plugins/Experimental/MovieSceneAnimMixer/Source/MovieSceneAnimMixerEditor/MovieSceneAnimMixerEditor.Build.cs
--- a/Engine/Plugins/Experimental/MovieSceneAnimMixer/Source/MovieSceneAnimMixerEditor/MovieSceneAnimMixerEditor.Build.cs
+++ b/Engine/Plugins/Experimental/MovieSceneAnimMixer/Source/MovieSceneAnimMixerEditor/MovieSceneAnimMixerEditor.Build.cs
@@ -9,21 +9,29 @@ namespace UnrealBuildTool.Rules
 			PublicDependencyModuleNames.AddRange(
 				new string[]
 				{
-					"Core",
-					"CoreUObject",
-					"Engine",
-					"MovieScene",
-					"AnimGraph",
 					"AnimationCore",
+					"AnimGraph",
 					"AnimGraphRuntime", 
-					"MovieSceneAnimMixer",
+					"Core",
+					"CoreUObject",
 				}
 			);
 			
 			PrivateDependencyModuleNames.AddRange(
 				new string[]
 				{
-					"BlueprintGraph"
+					"AssetRegistry",
+					"BlueprintGraph",
+					"Engine",
+					"MovieScene",
+					"MovieSceneAnimMixer",
+					"MovieSceneTracks",
+					"MovieSceneTools",
+					"Sequencer",
+					"SequencerCore",
+					"Slate",
+					"SlateCore",
+					"UnrealEd",
 				}
 			);
 		}
diff --git a/Engine/Plugins/Experimental/MovieSceneAnimMixer/Source/MovieSceneAnimMixerEditor/Private/MovieSceneAnimMixerEditorModule.cpp b/Engine/Plugins/Experimental/MovieSceneAnimMixer/Source/MovieSceneAnimMixerEditor/Private/MovieSceneAnimMixerEditorModule.cpp
--- a/Engine/Plugins/Experimental/MovieSceneAnimMixer/Source/MovieSceneAnimMixerEditor/Private/MovieSceneAnimMixerEditorModule.cpp
+++ b/Engine/Plugins/Experimental/MovieSceneAnimMixer/Source/MovieSceneAnimMixerEditor/Private/MovieSceneAnimMixerEditorModule.cpp
@@ -2,19 +2,29 @@
 
 #include "MovieSceneAnimMixerEditorModule.h"
 
+#include "ISequencerModule.h"
 #include "Modules/ModuleInterface.h"
 #include "Modules/ModuleManager.h"
+#include "MovieSceneAnimationMixerTrackEditor.h"
 
 namespace UE::MovieScene
 {
+	FLazyName SequencerModuleName("Sequencer");
+
 	void FMovieSceneAnimMixerEditorModule::StartupModule()
 	{
-		
+		using namespace UE::Sequencer;
+
+		ISequencerModule& SequencerModule = FModuleManager::Get().LoadModuleChecked<ISequencerModule>(SequencerModuleName.Resolve());
+		AnimationTrackEditorHandle = SequencerModule.RegisterTrackEditor(FOnCreateTrackEditor::CreateStatic( &FAnimationMixerTrackEditor::CreateTrackEditor));
 	}
 
 	void FMovieSceneAnimMixerEditorModule::ShutdownModule()
 	{
-		
+		if (ISequencerModule* SequencerModule = FModuleManager::Get().GetModulePtr<ISequencerModule>(SequencerModuleName.Resolve()))
+		{
+			SequencerModule->UnregisterTrackModel(AnimationTrackEditorHandle);
+		}
 	}
 }
 
diff --git a/Engine/Plugins/Experimental/MovieSceneAnimMixer/Source/MovieSceneAnimMixerEditor/Private/MovieSceneAnimMixerEditorModule.h b/Engine/Plugins/Experimental/MovieSceneAnimMixer/Source/MovieSceneAnimMixerEditor/Private/MovieSceneAnimMixerEditorModule.h
--- a/Engine/Plugins/Experimental/MovieSceneAnimMixer/Source/MovieSceneAnimMixerEditor/Private/MovieSceneAnimMixerEditorModule.h
+++ b/Engine/Plugins/Experimental/MovieSceneAnimMixer/Source/MovieSceneAnimMixerEditor/Private/MovieSceneAnimMixerEditorModule.h
@@ -14,6 +14,9 @@ namespace UE::MovieScene
 		// IModuleInterface interface
 		virtual void StartupModule() override;
 		virtual void ShutdownModule() override;
+
+	private:
+		FDelegateHandle AnimationTrackEditorHandle;
 	};
 }
 
diff --git a/Engine/Plugins/Experimental/MovieSceneAnimMixer/Source/MovieSceneAnimMixerEditor/Private/MovieSceneAnimationMixerTrackEditor.cpp b/Engine/Plugins/Experimental/MovieSceneAnimMixer/Source/MovieSceneAnimMixerEditor/Private/MovieSceneAnimationMixerTrackEditor.cpp
new file mode 100644
--- /dev/null
+++ b/Engine/Plugins/Experimental/MovieSceneAnimMixer/Source/MovieSceneAnimMixerEditor/Private/MovieSceneAnimationMixerTrackEditor.cpp
@@ -0,0 +1,284 @@
+// Copyright Epic Games, Inc. All Rights Reserved.
+
+#include "MovieSceneAnimationMixerTrackEditor.h"
+
+#include "AssetRegistry/AssetRegistryModule.h"
+#include "ClassIconFinder.h"
+#include "Components/SkeletalMeshComponent.h"
+#include "Framework/MultiBox/MultiBoxBuilder.h"
+#include "GameFramework/Actor.h"
+#include "ISequencer.h"
+#include "MovieSceneAnimationMixerTrack.h"
+#include "MVVM/Extensions/IObjectBindingExtension.h"
+#include "MVVM/Extensions/ITrackExtension.h"
+#include "MVVM/ViewModels/OutlinerColumns/OutlinerColumnTypes.h"
+#include "MVVM/ViewModels/SequencerEditorViewModel.h"
+#include "MVVM/ViewModels/TrackModel.h"
+#include "MVVM/Views/SOutlinerItemViewBase.h"
+#include "MVVM/Views/ViewUtilities.h"
+#include "ScopedTransaction.h"
+#include "Sections/MovieSceneSkeletalAnimationSection.h"
+#include "SequencerUtilities.h"
+
+#define LOCTEXT_NAMESPACE "MovieSceneAnimationMixerTrackEditor"
+
+
+namespace UE::Sequencer
+{
+
+TSharedRef<ISequencerTrackEditor> FAnimationMixerTrackEditor::CreateTrackEditor(TSharedRef<ISequencer> InSequencer)
+{
+	return MakeShareable(new FAnimationMixerTrackEditor(InSequencer));
+}
+
+FAnimationMixerTrackEditor::FAnimationMixerTrackEditor(TSharedRef<ISequencer> InSequencer)
+	: FCommonAnimationTrackEditor(InSequencer)
+{
+}
+
+FText FAnimationMixerTrackEditor::GetDisplayName() const
+{
+	return LOCTEXT("AnimationMixerTrackEditor_DisplayName", "Animation Mixer");
+}
+
+void FAnimationMixerTrackEditor::BuildObjectBindingTrackMenu(FMenuBuilder& MenuBuilder, const TArray<FGuid>& ObjectBindings, const UClass* ObjectClass)
+{
+	if (ObjectClass != nullptr && (ObjectClass->IsChildOf(USkeletalMeshComponent::StaticClass()) || ObjectClass->IsChildOf(AActor::StaticClass())))
+	{
+		MenuBuilder.AddMenuEntry(
+			LOCTEXT("AddAnimationTrack", "Animation Mixer"),
+			LOCTEXT("AddAnimationTrackTooltip", "Adds a new animation track for playing back Anim Sequences and other sources of animation."),
+			FSlateIcon(FAppStyle::GetAppStyleSetName(), "Sequencer.Tracks.Animation"),
+			FUIAction( 
+				FExecuteAction::CreateSP(this, &FAnimationMixerTrackEditor::HandleAddAnimationTrackMenuEntryExecute, ObjectBindings)
+			)
+		);
+	}
+}
+
+TSharedPtr<SWidget> FAnimationMixerTrackEditor::BuildOutlinerColumnWidget(const FBuildColumnWidgetParams& Params, const FName& ColumnName)
+{
+	using namespace UE::Sequencer;
+
+	TViewModelPtr<FSequencerEditorViewModel> Editor       = Params.Editor->CastThisShared<FSequencerEditorViewModel>();
+	TViewModelPtr<IOutlinerExtension>        OutlinerItem = Params.ViewModel.ImplicitCast();
+	if (!Editor || !OutlinerItem)
+	{
+		return SNullWidget::NullWidget;
+	}
+
+	if (ColumnName == FCommonOutlinerNames::Add)
+	{
+		return MakeAddButton(
+			LOCTEXT("AddSection", "Section"),
+			FOnGetContent::CreateSP(
+				this,
+				&FAnimationMixerTrackEditor::BuildAddSectionSubMenu,
+				TWeakViewModelPtr<IOutlinerExtension>(OutlinerItem),
+				TWeakViewModelPtr<FSequencerEditorViewModel>(Editor)
+			),
+			Params.ViewModel);
+	}
+
+	return FMovieSceneTrackEditor::BuildOutlinerColumnWidget(Params, ColumnName);
+}
+
+bool FAnimationMixerTrackEditor::SupportsSequence(UMovieSceneSequence* InSequence) const
+{
+	return true;
+}
+
+bool FAnimationMixerTrackEditor::SupportsType(TSubclassOf<UMovieSceneTrack> Type) const
+{
+	return (Type == UMovieSceneAnimationMixerTrack::StaticClass());
+}
+
+const FSlateBrush* FAnimationMixerTrackEditor::GetIconBrush() const
+{
+	return FAppStyle::GetBrush("Sequencer.Tracks.Animation");
+}
+
+TSubclassOf<UMovieSceneCommonAnimationTrack> FAnimationMixerTrackEditor::GetTrackClass() const
+{
+	return UMovieSceneAnimationMixerTrack::StaticClass();
+}
+
+TSharedRef<ISequencerSection> FAnimationMixerTrackEditor::MakeSectionInterface(UMovieSceneSection& SectionObject, UMovieSceneTrack& Track, FGuid ObjectBinding)
+{
+	check( SupportsType( SectionObject.GetOuter()->GetClass() ) );
+	
+	if (SectionObject.IsA<UMovieSceneSkeletalAnimationSection>())
+	{
+		return MakeShared<FCommonAnimationSection>(SectionObject, GetSequencer());
+	}
+
+	return FMovieSceneTrackEditor::MakeSectionInterface(SectionObject, Track, ObjectBinding);
+}
+
+TSharedRef<SWidget> FAnimationMixerTrackEditor::BuildAddSectionSubMenu(TWeakViewModelPtr<IOutlinerExtension> WeakViewModel, TWeakViewModelPtr<FSequencerEditorViewModel> WeakEditor)
+{
+	TViewModelPtr<ITrackExtension> Track = WeakViewModel.ImplicitPin();
+	if (!Track)
+	{
+		return SNullWidget::NullWidget;
+	}
+
+	auto CreateNewSection = [WeakViewModel, WeakEditor](FTopLevelAssetPath ClassPath)
+	{
+		UClass* SectionClass = FSoftClassPath(ClassPath.ToString()).TryLoadClass<UMovieSceneSection>();
+		if (!SectionClass)
+		{
+			return;
+		}
+
+		check(SectionClass && SectionClass->IsChildOf(UMovieSceneSection::StaticClass()) && SectionClass->ImplementsInterface(IMovieSceneAnimationSectionInterface::UClassType::StaticClass()));
+
+		TViewModelPtr<ITrackExtension>           Track     = WeakViewModel.ImplicitPin();
+		TViewModelPtr<FSequencerEditorViewModel> Editor    = WeakEditor.Pin();
+		TSharedPtr<ISequencer>                   Sequencer = Editor ? Editor->GetSequencer() : nullptr;
+
+		UMovieSceneTrack* TrackObject = Track ? Track->GetTrack() : nullptr;
+		if (Sequencer && TrackObject)
+		{
+			FScopedTransaction Transaction(FText::Format(LOCTEXT("AddSectionTransaction", "Add New {0} Section"), SectionClass->GetDisplayNameText()));
+
+			TrackObject->Modify();
+
+			UMovieSceneSection* NewSection = NewObject<UMovieSceneSection>(TrackObject, SectionClass, NAME_None, RF_Transactional);
+
+			int32 MaxRowIndex = -1;
+			for (UMovieSceneSection* Section : TrackObject->GetAllSections())
+			{
+				if (Section->GetClass() == SectionClass)
+				{
+					if (!Section->GetRange().Overlaps(NewSection->GetRange()))
+					{
+						if (MaxRowIndex > Section->GetRowIndex() || MaxRowIndex == -1)
+						{
+							MaxRowIndex = Section->GetRowIndex();
+						}
+					}
+					else if (MaxRowIndex == -1 || MaxRowIndex > Section->GetRowIndex()+1)
+					{
+						MaxRowIndex = Section->GetRowIndex() + 1;
+					}
+				}
+			}
+
+			NewSection->SetRowIndex(MaxRowIndex);
+
+			TrackObject->AddSection(*NewSection);
+
+			TrackObject->FixRowIndices();
+			TrackObject->UpdateEasing();
+
+			Sequencer->NotifyMovieSceneDataChanged(EMovieSceneDataChangeType::MovieSceneStructureItemAdded);
+		}
+	};
+
+	TViewModelSubListIterator<FSectionModel> Sections = Track->GetSectionModels().IterateSubList<FSectionModel>();
+
+	FMenuBuilder MenuBuilder(true, nullptr);
+
+	if (Sections && Cast<UMovieSceneSkeletalAnimationSection>(Sections->GetSection()))
+	{
+		MenuBuilder.BeginSection(NAME_None, LOCTEXT("TimeWarp_Label", "Time Warp"));
+		{
+			FSequencerUtilities::MakeTimeWarpMenuEntry(MenuBuilder, Track);
+		}
+		MenuBuilder.EndSection();
+	}
+
+	// Find all classes that match our interface
+	FAssetRegistryModule& AssetRegistryModule = FModuleManager::LoadModuleChecked<FAssetRegistryModule>("AssetRegistry");
+
+	TSet<FTopLevelAssetPath> AllClasses;
+	{
+		FTopLevelAssetPath TargetClassPath(UMovieSceneSection::StaticClass());
+
+		AssetRegistryModule.Get().GetDerivedClassNames({ TargetClassPath }, TSet<FTopLevelAssetPath>(), AllClasses);
+	}
+
+	for (auto It = AllClasses.CreateIterator(); It; ++It)
+	{
+		const UClass* Class = FSoftClassPath(It->ToString()).TryLoadClass<UMovieSceneSection>();
+
+		if (!Class || Class->HasMetaData("Hidden") || !Class->ImplementsInterface(IMovieSceneAnimationSectionInterface::UClassType::StaticClass()))
+		{
+			It.RemoveCurrent();
+		}
+	}
+
+	if (AllClasses.Num() != 0)
+	{
+		MenuBuilder.BeginSection(NAME_None, LOCTEXT("AnimationCategoryLabel", "Animation:"));
+
+		for (const FTopLevelAssetPath& ClassPath : AllClasses)
+		{
+			const UClass* Class = FSoftClassPath(ClassPath.ToString()).ResolveClass();
+			// Should have been loaded by the validation loop above
+			check(Class);
+
+			MenuBuilder.AddMenuEntry(
+				Class->GetDisplayNameText(),
+				Class->GetToolTipText(),
+				FSlateIcon(),
+				FUIAction(FExecuteAction::CreateLambda(CreateNewSection, ClassPath))
+			);
+		}
+
+		FViewModelPtr Model = WeakViewModel.Pin();
+		TViewModelPtr<IObjectBindingExtension> ObjectBinding = Model ? Model->FindAncestorOfType<IObjectBindingExtension>() : nullptr;
+		if (Track && ObjectBinding)
+		{
+			USkeleton* Skeleton = AcquireSkeletonFromObjectGuid(ObjectBinding->GetObjectGuid(), GetSequencer());
+			if (Skeleton)
+			{
+				MenuBuilder.AddSubMenu(
+					LOCTEXT("AddAnimationSubMenu", "Animation"),
+					LOCTEXT("AddAnimationSubMenu_Tooltip", "Adds a new animation section for an animation asset at the current time"),
+					FNewMenuDelegate::CreateSP(this, &FAnimationMixerTrackEditor::AddAnimationSubMenu, TArray<FGuid>({ ObjectBinding->GetObjectGuid() }), Skeleton, Track->GetTrack())
+				);
+			}
+		}
+
+		MenuBuilder.EndSection();
+	}
+
+	return MenuBuilder.MakeWidget();
+}
+
+void FAnimationMixerTrackEditor::HandleAddAnimationTrackMenuEntryExecute(TArray<FGuid> ObjectBindings)
+{
+	UMovieScene* MovieScene = GetFocusedMovieScene();
+
+	if (MovieScene == nullptr)
+	{
+		return;
+	}
+
+	if (MovieScene->IsReadOnly())
+	{
+		return;
+	}
+
+	const FScopedTransaction Transaction(LOCTEXT("AddAnimationTrack_Transaction", "Add Animation Track"));
+
+	MovieScene->Modify();
+
+	for (const FGuid& Guid : ObjectBindings)
+	{
+		FFindOrCreateTrackResult TrackResult = FindOrCreateTrackForObject(Guid, UMovieSceneAnimationMixerTrack::StaticClass());
+		if (TrackResult.bWasCreated)
+		{
+			UMovieSceneAnimationMixerTrack* AnimTrack = CastChecked<UMovieSceneAnimationMixerTrack>(TrackResult.Track);
+			AnimTrack->SetDisplayName(AnimTrack->GetDefaultDisplayName());
+		}
+	}
+}
+
+} // namespace UE::Sequencer
+
+#undef LOCTEXT_NAMESPACE
+
+
diff --git a/Engine/Plugins/Experimental/MovieScenePoseSearchTracks/Source/MovieScenePoseSearchTracks/Private/Sections/MovieSceneStitchAnimSection.cpp b/Engine/Plugins/Experimental/MovieScenePoseSearchTracks/Source/MovieScenePoseSearchTracks/Private/Sections/MovieSceneStitchAnimSection.cpp
--- a/Engine/Plugins/Experimental/MovieScenePoseSearchTracks/Source/MovieScenePoseSearchTracks/Private/Sections/MovieSceneStitchAnimSection.cpp
+++ b/Engine/Plugins/Experimental/MovieScenePoseSearchTracks/Source/MovieScenePoseSearchTracks/Private/Sections/MovieSceneStitchAnimSection.cpp
@@ -27,7 +27,6 @@
 
 #define LOCTEXT_NAMESPACE "MovieSceneStitchAnimSection"
 
-
 void FMovieSceneStitchAnimEvaluationTask::Execute(UE::AnimNext::FEvaluationVM& VM) const
 {
 	using namespace UE::AnimNext;
@@ -133,44 +132,47 @@ void FMovieSceneStitchAnimEvaluationTask::Execute(UE::AnimNext::FEvaluationVM& V
 	// We do this by aligning and projecting the matched animation asset twice,
 	// first to the 'initial' transform, and then to the 'target' transform, and blending the results based on time proximity.
 
-	FTransform InitialAnimSpaceRootTransform = AnimSequence->ExtractRootTrackTransform(FAnimExtractContext((double)MatchedAssetTime), nullptr);
-	FTransform CurrentAnimSpaceRootTransform = AnimSequence->ExtractRootTrackTransform(FAnimExtractContext((double)AnimationTime), nullptr);
-	FTransform TargetAnimSpaceRootTransform = AnimSequence->ExtractRootTrackTransform(FAnimExtractContext((double)(AnimationTime + (TimeToTarget * MatchedAssetPlayRate))), nullptr);
+	struct
+	{
+		FTransform Initial;
+		FTransform Current;
+		FTransform Target;
+	} AnimSpace;
+
+	AnimSpace.Initial = AnimSequence->ExtractRootTrackTransform(FAnimExtractContext((double)MatchedAssetTime), nullptr);
+	AnimSpace.Current = AnimSequence->ExtractRootTrackTransform(FAnimExtractContext((double)AnimationTime), nullptr);
+	AnimSpace.Target = AnimSequence->ExtractRootTrackTransform(FAnimExtractContext((double)(AnimationTime + (TimeToTarget * MatchedAssetPlayRate))), nullptr);
 
 	if (InverseMeshToActorRotation)
 	{
-		InitialAnimSpaceRootTransform.SetTranslation(InverseMeshToActorRotation.GetValue() * InitialAnimSpaceRootTransform.GetTranslation());
-		CurrentAnimSpaceRootTransform.SetTranslation(InverseMeshToActorRotation.GetValue() * CurrentAnimSpaceRootTransform.GetTranslation());
-		TargetAnimSpaceRootTransform.SetTranslation(InverseMeshToActorRotation.GetValue() * TargetAnimSpaceRootTransform.GetTranslation());
+		AnimSpace.Initial.SetTranslation(InverseMeshToActorRotation.GetValue() * AnimSpace.Initial.GetTranslation());
+		AnimSpace.Current.SetTranslation(InverseMeshToActorRotation.GetValue() * AnimSpace.Current.GetTranslation());
+		AnimSpace.Target.SetTranslation(InverseMeshToActorRotation.GetValue() * AnimSpace.Target.GetTranslation());
 	}
 
 	FTransform TargetTransformWorldSpace = StitchData.TargetTransform * SequenceTransformOrigin;
 
-	FTransform WorldSpaceProjectedFromStart = CurrentAnimSpaceRootTransform.GetRelativeTransform(InitialAnimSpaceRootTransform) * InitialRootTransform;
-	FTransform WorldSpaceProjectedFromTarget = CurrentAnimSpaceRootTransform.GetRelativeTransform(TargetAnimSpaceRootTransform) * TargetTransformWorldSpace;
+	FTransform WorldSpaceProjectedFromStart = AnimSpace.Current.GetRelativeTransform(AnimSpace.Initial) * InitialRootTransform;
+	FTransform WorldSpaceProjectedFromTarget = AnimSpace.Current.GetRelativeTransform(AnimSpace.Target) * TargetTransformWorldSpace;
 
 	FTransform RootTransformWorldSpace;
 	// TODO: Protect against division by 0
 	float BlendWeight = ((CurrentTime - InitialTime) / (TimeToTarget + CurrentTime - InitialTime));
 	RootTransformWorldSpace.Blend(WorldSpaceProjectedFromStart, WorldSpaceProjectedFromTarget, BlendWeight);
 
+	FQuat NewRotation = FQuat::Slerp(WorldSpaceProjectedFromStart.GetRotation(), WorldSpaceProjectedFromTarget.GetRotation(), BlendWeight);
+	RootTransformWorldSpace.SetRotation(NewRotation);
+
 	FDeltaTimeRecord DeltaTime;
 	DeltaTime.Set(PreviousAnimationTime, AnimationTime-PreviousAnimationTime);
 	const FAnimExtractContext ExtractionContext((double)AnimationTime, false, DeltaTime, false);
-		
+
 	if (EnumHasAnyFlags(VM.GetFlags(), EEvaluationFlags::Bones))
 	{
 		FDecompressionTools::GetAnimationPose(AnimSequence, ExtractionContext, Keyframe.Pose, bUseRawData);
 
-		int32 RootIndex = Keyframe.Pose.GetRefPose().GetLODBoneIndexFromSkeletonBoneIndex(0);
-		if (StitchData.SwapRootBone == ESwapRootBone::SwapRootBone_None)
-		{
-			Keyframe.Pose.LocalTransformsView[RootIndex] = RootTransformWorldSpace;
-		}
-		else
-		{
-			Keyframe.Pose.LocalTransformsView[RootIndex] = FTransform::Identity;
-		}
+		const int32 RootIndex = Keyframe.Pose.GetRefPose().GetLODBoneIndexFromSkeletonBoneIndex(0);
+		Keyframe.Pose.LocalTransformsView[RootIndex] = FTransform::Identity;
 	}
 
 	if (EnumHasAnyFlags(VM.GetFlags(), EEvaluationFlags::Curves))
@@ -187,6 +189,9 @@ void FMovieSceneStitchAnimEvaluationTask::Execute(UE::AnimNext::FEvaluationVM& V
 	{
 		FTransformAnimationAttribute* RootMotionAttribute = Keyframe.Attributes.FindOrAdd<FTransformAnimationAttribute>(AnimMixerComponents->RootTransformAttributeId);
 		RootMotionAttribute->Value = RootTransformWorldSpace;
+
+		FFloatAnimationAttribute* RootMotionWeightAttribute = Keyframe.Attributes.FindOrAdd<FFloatAnimationAttribute>(AnimMixerComponents->RootTransformWeightAttributeId);
+		RootMotionWeightAttribute->Value = 1.f;
 	}
 
 	VM.PushValue(KEYFRAME_STACK_NAME, MakeUnique<FKeyframeState>(MoveTemp(Keyframe)));
@@ -254,7 +259,6 @@ void UMovieSceneStitchAnimSection::ImportEntityImpl(UMovieSceneEntitySystemLinke
 	StitchData.TargetTransform = TargetTransform;
 	StitchData.StartFrame = GetInclusiveStartFrame();
 	StitchData.EndFrame = GetExclusiveEndFrame();
-	StitchData.SwapRootBone = SwapRootBone;
 	
 	// Make a default eval task. This will be filled out more by the stitch system
 	TSharedPtr<FMovieSceneStitchAnimEvaluationTask> EvalTask = MakeShared<FMovieSceneStitchAnimEvaluationTask>();
@@ -270,7 +274,7 @@ void UMovieSceneStitchAnimSection::ImportEntityImpl(UMovieSceneEntitySystemLinke
 		.Add(AnimMixerComponents->Priority, MixedAnimationPriority)
 		.Add(AnimMixerComponents->Target, MixedAnimationTarget)
 		.Add(AnimMixerComponents->Task, EvalTask)
-		.Add(AnimMixerComponents->RootMotionSettings, FMovieSceneRootMotionSettings{StitchData.SwapRootBone, EMovieSceneRootMotionSpace::WorldSpace})
+		.Add(AnimMixerComponents->RootMotionSettings, FMovieSceneRootMotionSettings{TOptional<FSkeletalAnimationRootMotionOverride>(), EMovieSceneRootMotionSpace::WorldSpace})
 		.AddTag(AnimMixerComponents->Tags.RequiresBlending)
 	);
 }
@@ -280,5 +284,5 @@ double FMovieSceneStitchAnimComponentData::MapTimeToSectionSeconds(FFrameTime In
 	return InFrameRate.AsSeconds(InPosition - StartFrame);
 }
 
-#undef LOCTEXT_NAMESPACE 
+#undef LOCTEXT_NAMESPACE
 
diff --git a/Engine/Plugins/Experimental/MovieScenePoseSearchTracks/Source/MovieScenePoseSearchTracks/Private/Systems/MovieSceneStitchAnimSystem.cpp b/Engine/Plugins/Experimental/MovieScenePoseSearchTracks/Source/MovieScenePoseSearchTracks/Private/Systems/MovieSceneStitchAnimSystem.cpp
--- a/Engine/Plugins/Experimental/MovieScenePoseSearchTracks/Source/MovieScenePoseSearchTracks/Private/Systems/MovieSceneStitchAnimSystem.cpp
+++ b/Engine/Plugins/Experimental/MovieScenePoseSearchTracks/Source/MovieScenePoseSearchTracks/Private/Systems/MovieSceneStitchAnimSystem.cpp
@@ -180,8 +180,6 @@ namespace UE::MovieScene
 					AnimTask->InitialTime = PreviousTimeSeconds;
 					AnimTask->SequenceTransformOrigin = GetTransformOrigin(InstanceHandle);
 
-					FTransform InitialComponentTransform = FTransform::Identity;
-
 					const AActor* Actor = Cast<AActor>(BoundObject);
 					if (!Actor)
 					{
@@ -199,21 +197,15 @@ namespace UE::MovieScene
 
 					AnimTask->ContextObject = GetContextObject(Targets[Index], BoundObject);
 
-					if (Actor && StitchAnim.SwapRootBone == ESwapRootBone::SwapRootBone_Actor)
+					if (Actor)
 					{
-						InitialComponentTransform = Actor->GetRootComponent()->GetRelativeTransform();
+						AnimTask->InitialRootTransform = Actor->GetRootComponent()->GetRelativeTransform();
 
 						if (SkeletalMeshComponent)
 						{
 							AnimTask->InverseMeshToActorRotation = Actor->GetRootComponent()->GetComponentTransform().GetRelativeTransformReverse(SkeletalMeshComponent->GetComponentTransform()).GetRotation();
 						}
 					}
-					else if (SkeletalMeshComponent && StitchAnim.SwapRootBone == ESwapRootBone::SwapRootBone_Component)
-					{
-						InitialComponentTransform = SkeletalMeshComponent->GetRelativeTransform();
-					}
-
-					AnimTask->InitialRootTransform = InitialComponentTransform;
 				}
 				AnimTask->PreviousTime = PreviousTimeSeconds;
 				AnimTask->CurrentTime = CurrentTimeSeconds;
diff --git a/Engine/Plugins/Experimental/MovieScenePoseSearchTracks/Source/MovieScenePoseSearchTracks/Public/Sections/MovieSceneStitchAnimSection.h b/Engine/Plugins/Experimental/MovieScenePoseSearchTracks/Source/MovieScenePoseSearchTracks/Public/Sections/MovieSceneStitchAnimSection.h
--- a/Engine/Plugins/Experimental/MovieScenePoseSearchTracks/Source/MovieScenePoseSearchTracks/Public/Sections/MovieSceneStitchAnimSection.h
+++ b/Engine/Plugins/Experimental/MovieScenePoseSearchTracks/Source/MovieScenePoseSearchTracks/Public/Sections/MovieSceneStitchAnimSection.h
@@ -9,6 +9,7 @@
 #include "EntitySystem/IMovieSceneEntityProvider.h"
 #include "StructUtils/InstancedStruct.h"
 #include "MovieSceneMixedAnimationTarget.h"
+#include "MovieSceneAnimationMixerTrack.h"
 #include "Animation/AnimationAsset.h"
 #include "EvaluationVM/EvaluationTask.h"
 #include "AnimSequencerInstanceProxy.h"
@@ -37,9 +38,6 @@ struct FMovieSceneStitchAnimComponentData
 	UPROPERTY()
 	FFrameNumber EndFrame;
 
-	UPROPERTY()
-	ESwapRootBone SwapRootBone = ESwapRootBone::SwapRootBone_None;
-
 	double MapTimeToSectionSeconds(FFrameTime InPosition, FFrameRate InFrameRate) const;
 };
 
@@ -99,10 +97,11 @@ struct FMovieSceneStitchAnimEvaluationTask : public FAnimNextEvaluationTask
 /**
  * Handles generating and playing back transitional skeletal animations from a stitch database.
  */
-UCLASS( MinimalAPI )
+UCLASS( MinimalAPI, DisplayName="Animation Stitch")
 class UMovieSceneStitchAnimSection
 	: public UMovieSceneSection
 	, public IMovieSceneEntityProvider
+	, public IMovieSceneAnimationSectionInterface
 {
 	GENERATED_BODY()
 
@@ -123,16 +122,6 @@ public:
 	UPROPERTY(EditAnywhere, Category = "Stitch")
 	FTransform TargetTransform;
 
-	// TODO: Would be nice to standardize these properties across multiple anim tracks
-	UPROPERTY(EditAnywhere, Category = "Root Motion")
-	bool bBlendFirstChildOfRoot = false;
-	
-	UPROPERTY(EditAnywhere, Category = "Root Motion")
-	int32 ChildBoneIndex = INDEX_NONE;
-
-	UPROPERTY(EditAnywhere, Category = "Root Motion")
-	ESwapRootBone SwapRootBone = ESwapRootBone::SwapRootBone_Actor;
-
 	UPROPERTY(EditAnywhere, Category = "Mixed Animation")
 	TInstancedStruct<FMovieSceneMixedAnimationTarget> MixedAnimationTarget;
 
@@ -147,6 +136,11 @@ public:
 	virtual EMovieSceneChannelProxyType CacheChannelProxy() override;
 	virtual UObject* GetSourceObject() const override;
 
+	virtual int32 GetRowSortOrder() const override
+	{
+		return 0;
+	}
+
 protected:
 
 	virtual float GetTotalWeightValue(FFrameTime InTime) const override;
diff --git a/Engine/Source/Editor/EditorStyle/Private/StarshipStyle.cpp b/Engine/Source/Editor/EditorStyle/Private/StarshipStyle.cpp
--- a/Engine/Source/Editor/EditorStyle/Private/StarshipStyle.cpp
+++ b/Engine/Source/Editor/EditorStyle/Private/StarshipStyle.cpp
@@ -2380,6 +2380,7 @@ void FStarshipEditorStyle::FStyle::SetupSequencerStyles()
 		Set("Sequencer.Tracks.CinematicShot", new IMAGE_BRUSH_SVG("Starship/Sequencer/ShotTrack", Icon16x16));
 		Set("Sequencer.Tracks.Slomo", new IMAGE_BRUSH_SVG("Starship/Sequencer/SlomoTrack", Icon16x16));
 		Set("Sequencer.Tracks.TimeWarp", new IMAGE_BRUSH_SVG("Starship/Sequencer/TimeWarp_16", Icon16x16));
+		Set("Sequencer.Tracks.Animation", new IMAGE_BRUSH_SVG("Starship/Sequencer/Animation", Icon16x16));
 		Set("Sequencer.Tracks.Sub", new IMAGE_BRUSH_SVG("Starship/Sequencer/SubTrack", Icon16x16));
 		Set("Sequencer.Tracks.LevelVisibility", new IMAGE_BRUSH_SVG("Starship/Sequencer/LevelVisibilityTrack", Icon16x16));
 		Set("Sequencer.Tracks.DataLayer", new IMAGE_BRUSH_SVG("Starship/Common/DataLayers", Icon16x16));
