Sequencer: Root Motion // AnimMixer fixes

- Added a root motion space setting to the root motion track to allow specification of world vs local space
- Fixed Evaluation Program holding a strong reference to the root motion parameters, resulting in mixers not being cleaned up correctly
- Actor transform is now removed from root bones when applied with RootBone (since root motion transforms are always in world space)
- Fixed failure to respect legacy swap root bone settings
- Added ability to keep track of which transforms came from sequencer tracks and which did not

diff --git a/Engine/Plugins/Experimental/MovieSceneAnimMixer/Source/MovieSceneAnimMixer/Internal/MovieSceneAnimationMixerTrack.h b/Engine/Plugins/Experimental/MovieSceneAnimMixer/Source/MovieSceneAnimMixer/Internal/MovieSceneAnimationMixerTrack.h
--- a/Engine/Plugins/Experimental/MovieSceneAnimMixer/Source/MovieSceneAnimMixer/Internal/MovieSceneAnimationMixerTrack.h
+++ b/Engine/Plugins/Experimental/MovieSceneAnimMixer/Source/MovieSceneAnimMixer/Internal/MovieSceneAnimationMixerTrack.h
@@ -4,6 +4,7 @@
 
 #include "CoreTypes.h"
 #include "Channels/MovieSceneDoubleChannel.h"
+#include "Channels/MovieSceneByteChannel.h"
 #include "Decorations/IMovieSceneChannelDecoration.h"
 #include "Decorations/IMovieSceneSectionDecoration.h"
 #include "EntitySystem/IMovieSceneEntityDecorator.h"
@@ -13,6 +14,17 @@
 #include "Tracks/MovieSceneCommonAnimationTrack.h"
 #include "MovieSceneAnimationMixerTrack.generated.h"
 
+
+// Enum describing what space that root motion should be applied in.
+UENUM()
+enum class EMovieSceneRootMotionSpace : uint8
+{
+	// Root motion should be applied in animation space, meaning that it will be applied on top of the blended transform track and transform origin.
+	AnimationSpace,
+	// Root motion should be applied in world space, meaning that it will override any transform track or transform origin.
+	WorldSpace
+};
+
 UINTERFACE(MinimalAPI)
 class UMovieSceneAnimationSectionInterface : public UInterface
 {
@@ -63,6 +75,9 @@ class UMovieSceneAnimationBaseTransformDecoration
 
 	UPROPERTY()
 	FMovieSceneDoubleChannel Rotation[3];
+
+	UPROPERTY()
+	FMovieSceneByteChannel RootMotionSpace;
 };
 
 
diff --git a/Engine/Plugins/Experimental/MovieSceneAnimMixer/Source/MovieSceneAnimMixer/Private/AnimMixerComponentTypes.cpp b/Engine/Plugins/Experimental/MovieSceneAnimMixer/Source/MovieSceneAnimMixer/Private/AnimMixerComponentTypes.cpp
--- a/Engine/Plugins/Experimental/MovieSceneAnimMixer/Source/MovieSceneAnimMixer/Private/AnimMixerComponentTypes.cpp
+++ b/Engine/Plugins/Experimental/MovieSceneAnimMixer/Source/MovieSceneAnimMixer/Private/AnimMixerComponentTypes.cpp
@@ -11,7 +11,6 @@ namespace UE::MovieScene
 {
 
 	const UE::Anim::FAttributeId FAnimMixerComponentTypes::RootTransformAttributeId(TEXT("RootTransform"), FCompactPoseBoneIndex(0));
-	const UE::Anim::FAttributeId FAnimMixerComponentTypes::RootTransformBaseAttributeId(TEXT("RootTransformBase"), FCompactPoseBoneIndex(0));
 	const UE::Anim::FAttributeId FAnimMixerComponentTypes::RootTransformWeightAttributeId(TEXT("RootTransformWeight"), FCompactPoseBoneIndex(0));
 	const UE::Anim::FAttributeId FAnimMixerComponentTypes::LegacySwapRootBoneAttributeId(TEXT("LegacySwapRootBone"), FCompactPoseBoneIndex(0));
 
diff --git a/Engine/Plugins/Experimental/MovieSceneAnimMixer/Source/MovieSceneAnimMixer/Private/MovieSceneAnimationMixerTrack.cpp b/Engine/Plugins/Experimental/MovieSceneAnimMixer/Source/MovieSceneAnimMixer/Private/MovieSceneAnimationMixerTrack.cpp
--- a/Engine/Plugins/Experimental/MovieSceneAnimMixer/Source/MovieSceneAnimMixer/Private/MovieSceneAnimationMixerTrack.cpp
+++ b/Engine/Plugins/Experimental/MovieSceneAnimMixer/Source/MovieSceneAnimMixer/Private/MovieSceneAnimationMixerTrack.cpp
@@ -16,6 +16,9 @@ UMovieSceneAnimationBaseTransformDecoration::UMovieSceneAnimationBaseTransformDe
 	Rotation[0].SetDefault(0.0);
 	Rotation[1].SetDefault(0.0);
 	Rotation[2].SetDefault(0.0);
+
+	RootMotionSpace.SetEnum(StaticEnum<EMovieSceneRootMotionSpace>());
+	RootMotionSpace.SetDefault((uint8)EMovieSceneRootMotionSpace::AnimationSpace);
 }
 
 EMovieSceneChannelProxyType UMovieSceneAnimationBaseTransformDecoration::PopulateChannelProxy(FMovieSceneChannelProxyData& OutProxyData)
@@ -73,6 +76,13 @@ EMovieSceneChannelProxyType UMovieSceneAnimationBaseTransformDecoration::Populat
 		ChannelMetaData.SortOrder = 5;
 		OutProxyData.Add(Rotation[2], ChannelMetaData, TMovieSceneExternalValue<double>());
 	}
+	{
+		FMovieSceneChannelMetaData ChannelMetaData;
+		ChannelMetaData.SetIdentifiers("Space", NSLOCTEXT("BaseTransformDecoration", "Space", "Space"));
+		ChannelMetaData.WeakOwningObject = this;
+		ChannelMetaData.SortOrder = 6;
+		OutProxyData.Add(RootMotionSpace, ChannelMetaData, TMovieSceneExternalValue<uint8>());
+	}
 
 #else
 
@@ -84,6 +94,8 @@ EMovieSceneChannelProxyType UMovieSceneAnimationBaseTransformDecoration::Populat
 	OutProxyData.Add(Rotation[1]);
 	OutProxyData.Add(Rotation[2]);
 
+	OutProxyData.Add(RootMotionSpace);
+
 #endif
 
 	return EMovieSceneChannelProxyType::Static;
@@ -102,6 +114,7 @@ void UMovieSceneAnimationBaseTransformDecoration::ExtendEntityImpl(UMovieSceneEn
 		.Add(BuiltInComponents->DoubleChannel[3], &Rotation[0])
 		.Add(BuiltInComponents->DoubleChannel[4], &Rotation[1])
 		.Add(BuiltInComponents->DoubleChannel[5], &Rotation[2])
+		.Add(BuiltInComponents->ByteChannel, &RootMotionSpace)
 	);
 }
 
diff --git a/Engine/Plugins/Experimental/MovieSceneAnimMixer/Source/MovieSceneAnimMixer/Private/Systems/MovieSceneAnimMixerSystem.cpp b/Engine/Plugins/Experimental/MovieSceneAnimMixer/Source/MovieSceneAnimMixer/Private/Systems/MovieSceneAnimMixerSystem.cpp
--- a/Engine/Plugins/Experimental/MovieSceneAnimMixer/Source/MovieSceneAnimMixer/Private/Systems/MovieSceneAnimMixerSystem.cpp
+++ b/Engine/Plugins/Experimental/MovieSceneAnimMixer/Source/MovieSceneAnimMixer/Private/Systems/MovieSceneAnimMixerSystem.cpp
@@ -3,6 +3,7 @@
 #include "Systems/MovieSceneAnimMixerSystem.h"
 
 #include "MovieSceneRootMotionSection.h"
+#include "Systems/ByteChannelEvaluatorSystem.h"
 
 #include "EntitySystem/BuiltInComponentTypes.h"
 #include "AnimMixerComponentTypes.h"
@@ -51,19 +52,26 @@ namespace UE::MovieScene
 		static void ForEachEntity(
 			TSharedPtr<FAnimNextEvaluationTask> Task,
 			TSharedPtr<FMovieSceneAnimMixerEntry> MixerEntry,
+			const FMovieSceneRootMotionSettings* RootMotionSettings,
 			const double* WeightAndEasings)
 		{
 			MixerEntry->EvalTask = Task;
 			MixerEntry->PoseWeight = WeightAndEasings ? *WeightAndEasings : 1.0;
+			if (RootMotionSettings)
+			{
+				MixerEntry->RootMotionSettings = *RootMotionSettings;
+			}
 		}
 	};
 
 	struct FEvaluateAnimMixers
 	{
 		TMap<FMovieSceneAnimMixerKey, TSharedPtr<FMovieSceneAnimMixer>>* Mixers;
+		UMovieSceneRootMotionSystem* RootMotionSystem;
 
-		FEvaluateAnimMixers(TMap<FMovieSceneAnimMixerKey, TSharedPtr<FMovieSceneAnimMixer>>* InMixers)
+		FEvaluateAnimMixers(TMap<FMovieSceneAnimMixerKey, TSharedPtr<FMovieSceneAnimMixer>>* InMixers, UMovieSceneRootMotionSystem* InRootMotionSystem)
 			: Mixers(InMixers)
+			, RootMotionSystem(InRootMotionSystem)
 		{
 		}
 
@@ -241,8 +249,14 @@ namespace UE::MovieScene
 
 			if (RootMotion)
 			{
+				USceneComponent* Component = Cast<USceneComponent>(MeshComponent.GetObject());
+				USceneComponent* Root      = Component ? Component->GetOwner()->GetRootComponent() : nullptr;
+
+				const bool bComponentHasKeyedTransform     = (Component && RootMotionSystem->IsTransformKeyed(Component));
+				const bool bRootComponentHasKeyedTransform = (Root      && RootMotionSystem->IsTransformKeyed(Root));
+
 				// Add task to store the final root motion result
-				Mixer->EvaluationProgram->AppendTask(FAnimNextStoreRootTransformTask::Make(RootMotion));
+				Mixer->EvaluationProgram->AppendTask(FAnimNextStoreRootTransformTask::Make(RootMotion, bComponentHasKeyedTransform, bRootComponentHasKeyedTransform));
 			}
 		}
 	};
@@ -315,6 +329,7 @@ UMovieSceneAnimMixerSystem::UMovieSceneAnimMixerSystem(const FObjectInitializer&
 		DefineImplicitPrerequisite(UMovieSceneBoundSceneComponentInstantiator::StaticClass(), GetClass());
 		DefineImplicitPrerequisite(UWeightAndEasingEvaluatorSystem::StaticClass(), GetClass());
 		DefineImplicitPrerequisite(UMovieSceneComponentTransformSystem::StaticClass(), GetClass());
+		DefineImplicitPrerequisite(UByteChannelEvaluatorSystem::StaticClass(), GetClass());
 		DefineComponentConsumer(GetClass(), BuiltInComponents->BoundObjectKey);
 		DefineComponentConsumer(GetClass(), BuiltInComponents->BoundObject);
 		DefineComponentConsumer(GetClass(), AnimMixerComponents->Task);
@@ -463,12 +478,33 @@ void UMovieSceneAnimMixerSystem::OnRun(FSystemTaskPrerequisites& InPrerequisites
 					if (!RootMotion)
 					{
 						RootMotion = MakeShared<FMovieSceneMixerRootMotionComponentData>();
+
 						RootMotion->RootDestination = EMovieSceneRootMotionDestination::RootBone;
+						RootMotion->OriginalBoundObject = Cast<USceneComponent>(BoundObjectKey.ResolveObjectPtr());
 
 						Mixer->WeakRootMotion = RootMotion;
 					}
+
+					// Overwrite the RootDestination for the root motion if we have a legacy setting.
+					//   If any actual RootDestination components exist, they will simply override this on eval
+					switch (RootMotionSettings->LegacySwapRootBone)
+					{
+					case ESwapRootBone::SwapRootBone_None:
+						break;
+					case ESwapRootBone::SwapRootBone_Component:
+						RootMotion->RootDestination = EMovieSceneRootMotionDestination::Component;
+						break;
+					case ESwapRootBone::SwapRootBone_Actor:
+						RootMotion->RootDestination = EMovieSceneRootMotionDestination::Actor;
+						break;
+					}
+
 					Entry->RootMotionLifetimeReference = RootMotion;
 				}
+				else
+				{
+					Entry->RootMotionLifetimeReference = nullptr;
+				}
 
 				const FComponentMask& Type = Linker->EntityManager.GetEntityType(EntityID);
 				Entry->bAdditive = Type.ContainsAny({BuiltInComponents->Tags.AdditiveBlend, BuiltInComponents->Tags.AdditiveFromBaseBlend});
@@ -583,6 +619,7 @@ void UMovieSceneAnimMixerSystem::OnSchedulePersistentTasks(UE::MovieScene::IEnti
 	FTaskID UpdateTask = FEntityTaskBuilder()
 		.Read(AnimMixerComponents->Task)
 		.Write(AnimMixerComponents->MixerEntry)
+		.ReadOptional(AnimMixerComponents->RootMotionSettings)
 		.ReadOptional(BuiltInComponents->WeightAndEasingResult)
 		.FilterNone({BuiltInComponents->Tags.NeedsUnlink})
 		.Schedule_PerEntity<FUpdateTaskEntities>(&Linker->EntityManager, TaskScheduler);
@@ -593,7 +630,7 @@ void UMovieSceneAnimMixerSystem::OnSchedulePersistentTasks(UE::MovieScene::IEnti
 		.Read(AnimMixerComponents->MeshComponent)
 		.Read(AnimMixerComponents->Target)
 		.Write(AnimMixerComponents->MixerTask)
-		.Schedule_PerEntity<FEvaluateAnimMixers>(&Linker->EntityManager, TaskScheduler, &Mixers);
+		.Schedule_PerEntity<FEvaluateAnimMixers>(&Linker->EntityManager, TaskScheduler, &Mixers, Linker->FindSystem<UMovieSceneRootMotionSystem>());
 
 	TaskScheduler->AddPrerequisite(UpdateTask, MixTask);
 }
diff --git a/Engine/Plugins/Experimental/MovieSceneAnimMixer/Source/MovieSceneAnimMixer/Private/Systems/MovieSceneMixedSkeletalAnimationSystem.cpp b/Engine/Plugins/Experimental/MovieSceneAnimMixer/Source/MovieSceneAnimMixer/Private/Systems/MovieSceneMixedSkeletalAnimationSystem.cpp
--- a/Engine/Plugins/Experimental/MovieSceneAnimMixer/Source/MovieSceneAnimMixer/Private/Systems/MovieSceneMixedSkeletalAnimationSystem.cpp
+++ b/Engine/Plugins/Experimental/MovieSceneAnimMixer/Source/MovieSceneAnimMixer/Private/Systems/MovieSceneMixedSkeletalAnimationSystem.cpp
@@ -111,17 +111,20 @@ struct FSkeletalAnimMixerMutation : IMovieSceneEntityMutation
 
 				OutPriorities[Index] = SkeletalAnim.Section->MixedAnimationPriority;
 				OutRootMotionSettings[Index].RootMotionSpace = EMovieSceneRootMotionSpace::AnimationSpace;
-				OutRootMotionSettings[Index].LegacySwapRootBone = SkeletalAnim.Section->Params.SwapRootBone;
+
 
 				bool bHasRootMotionOverride = false;
 				UMovieSceneSkeletalAnimationTrack* Track = SkeletalAnim.Section->GetTypedOuter<UMovieSceneSkeletalAnimationTrack>();
-				if (Track)
+
+				ESwapRootBone LegacySwapRootBone = SkeletalAnim.Section->Params.SwapRootBone;
+				if (LegacySwapRootBone == ESwapRootBone::SwapRootBone_None && Track)
 				{
-					if (OutRootMotionSettings[Index].LegacySwapRootBone == ESwapRootBone::SwapRootBone_None)
-					{
-						OutRootMotionSettings[Index].LegacySwapRootBone = Track->SwapRootBone;
-					}
+					LegacySwapRootBone = Track->SwapRootBone;
+				}
+				OutRootMotionSettings[Index].LegacySwapRootBone = LegacySwapRootBone;
 
+				if (Track)
+				{
 					if (Track->RootMotionParams.bRootMotionsDirty) 
 					{
 						constexpr bool bForce = true;
@@ -151,9 +154,15 @@ struct FGatherMixableRootMotion
 	static void ForEachEntity(
 		const FMovieSceneSkeletalAnimationComponentData& SkeletalAnimation,
 		FFrameTime EvalTime,
+		const uint8* RootMotionSpace,
 		FMovieSceneRootMotionSettings& RootMotionSettings
 		)
 	{
+		if (RootMotionSpace)
+		{
+			RootMotionSettings.RootMotionSpace = (EMovieSceneRootMotionSpace)*RootMotionSpace;
+		}
+
 		if (RootMotionSettings.RootMotionOverride.IsSet())
 		{
 			UMovieSceneSkeletalAnimationSection::FRootMotionParams RootMotionParams;
@@ -294,16 +303,13 @@ struct FGatherMixableSkeletalAnimations
 			if (OptionalRootMotionSettings)
 			{
 				AnimData.RootMotionOverride = OptionalRootMotionSettings[Index].RootMotionOverride;
+				AnimData.LegacySwapRootBone = OptionalRootMotionSettings[Index].LegacySwapRootBone;
 			}
 			AnimData.FromPosition = PreviousEvalTime;
 			AnimData.ToPosition = EvalTime;
 			AnimData.bFireNotifies = bFireNotifies && !AnimParams.bSkipAnimNotifiers && !bLooped;
 			// TODO: We want to allow for additive to be set on the section properties
 			AnimData.bAdditive = AnimData.AnimSequence.IsValid() ? AnimData.AnimSequence->IsValidAdditive() : false;
-			if (OptionalRootMotionSettings)
-			{
-				AnimData.LegacySwapRootBone = OptionalRootMotionSettings->LegacySwapRootBone;
-			}
 
 			// TODO: Need to create a mirroring task and add that if appropriate
 			// TODO: Notifies- perhaps this should be a part of the animnext side of things?
@@ -363,6 +369,7 @@ void FMovieSceneSkeletalAnimationEvaluationTask::Execute(UE::AnimNext::FEvaluati
 		const bool bUseRootMotion = AnimSequencePtr->bEnableRootMotion;
 
 		TOptional<FTransform> RootBoneTransform;
+		const int32 RootIndex = Keyframe.Pose.GetRefPose().GetLODBoneIndexFromSkeletonBoneIndex(0);
 
 		if (EnumHasAnyFlags(VM.GetFlags(), EEvaluationFlags::Bones))
 		{
@@ -388,9 +395,7 @@ void FMovieSceneSkeletalAnimationEvaluationTask::Execute(UE::AnimNext::FEvaluati
 			// Extract the root motion if we've chosen that option
 			if (bUseRootMotion)
 			{
-				int32 RootIndex = Keyframe.Pose.GetRefPose().GetLODBoneIndexFromSkeletonBoneIndex(0);
 				RootBoneTransform = Keyframe.Pose.LocalTransformsView[RootIndex];
-				Keyframe.Pose.LocalTransformsView[RootIndex] = FTransform::Identity;
 			}
 		}
 
@@ -435,12 +440,12 @@ void FMovieSceneSkeletalAnimationEvaluationTask::Execute(UE::AnimNext::FEvaluati
 
 		if (RootBoneTransform.IsSet())
 		{
-			FTransformAnimationAttribute* RootTransformBaseAttribute = Keyframe.Attributes.FindOrAdd<FTransformAnimationAttribute>(AnimMixerComponents->RootTransformBaseAttributeId);
+			Keyframe.Pose.LocalTransformsView[RootIndex] = FTransform::Identity;
+
 			FTransformAnimationAttribute* RootMotionAttribute        = Keyframe.Attributes.FindOrAdd<FTransformAnimationAttribute>(AnimMixerComponents->RootTransformAttributeId);
 			FFloatAnimationAttribute*     RootMotionWeight           = Keyframe.Attributes.FindOrAdd<FFloatAnimationAttribute>(AnimMixerComponents->RootTransformWeightAttributeId);
 
-			RootTransformBaseAttribute->Value = RootBaseTransform;
-			RootMotionAttribute->Value        = RootBoneTransform.GetValue();
+			RootMotionAttribute->Value        = RootBoneTransform.GetValue() * RootBaseTransform;
 			RootMotionWeight->Value           = 1.f;
 
 			if (AnimationData.LegacySwapRootBone != ESwapRootBone::SwapRootBone_None)
@@ -528,6 +533,7 @@ void UMovieSceneMixedSkeletalAnimationSystem::OnSchedulePersistentTasks(UE::Movi
 	FTaskID GatherRootMotionTask = FEntityTaskBuilder()
 		.Read(TrackComponents->SkeletalAnimation)
 		.Read(BuiltInComponents->EvalTime)
+		.ReadOptional(BuiltInComponents->ByteResult)      // Root Space
 		.Write(AnimMixerComponents->RootMotionSettings)
 		.SetDesiredThread(ENamedThreads::GameThread_Local)
 		.FilterNone({ BuiltInComponents->Tags.Ignored })
@@ -541,12 +547,12 @@ void UMovieSceneMixedSkeletalAnimationSystem::OnSchedulePersistentTasks(UE::Movi
 		.Read(TrackComponents->SkeletalAnimation)
 		.ReadOptional(BuiltInComponents->EvalTime)
 		.ReadOptional(AnimMixerComponents->RootMotionSettings)
-		.ReadOptional(BuiltInComponents->DoubleResult[0])
-		.ReadOptional(BuiltInComponents->DoubleResult[1])
-		.ReadOptional(BuiltInComponents->DoubleResult[2])
-		.ReadOptional(BuiltInComponents->DoubleResult[3])
-		.ReadOptional(BuiltInComponents->DoubleResult[4])
-		.ReadOptional(BuiltInComponents->DoubleResult[5])
+		.ReadOptional(BuiltInComponents->DoubleResult[0]) // Root Offset X
+		.ReadOptional(BuiltInComponents->DoubleResult[1]) // Root Offset Y
+		.ReadOptional(BuiltInComponents->DoubleResult[2]) // Root Offset Z
+		.ReadOptional(BuiltInComponents->DoubleResult[3]) // Root Rotation Offset X
+		.ReadOptional(BuiltInComponents->DoubleResult[4]) // Root Rotation Offset Y
+		.ReadOptional(BuiltInComponents->DoubleResult[5]) // Root Rotation Offset Z
 		.Write(AnimMixerComponents->Task)
 		.FilterAll( { TrackComponents->Tags.AnimMixerPoseProducer })
 		.FilterNone({ BuiltInComponents->Tags.Ignored})
diff --git a/Engine/Plugins/Experimental/MovieSceneAnimMixer/Source/MovieSceneAnimMixer/Private/Systems/MovieSceneRootMotionSystem.cpp b/Engine/Plugins/Experimental/MovieSceneAnimMixer/Source/MovieSceneAnimMixer/Private/Systems/MovieSceneRootMotionSystem.cpp
--- a/Engine/Plugins/Experimental/MovieSceneAnimMixer/Source/MovieSceneAnimMixer/Private/Systems/MovieSceneRootMotionSystem.cpp
+++ b/Engine/Plugins/Experimental/MovieSceneAnimMixer/Source/MovieSceneAnimMixer/Private/Systems/MovieSceneRootMotionSystem.cpp
@@ -136,18 +136,22 @@ void FMovieSceneMixerRootMotionComponentData::Initialize()
 		return;
 	}
 
+	ComponentToWorldTransform = BoundObject->GetComponentToWorld();
+
 	// If we're applying to the actor or a root custom attribute, we need to factor out the actor->compont rotation
 	const bool bNeedInverseMeshRotation = RootDestination == EMovieSceneRootMotionDestination::Actor
 			|| RootDestination == EMovieSceneRootMotionDestination::Attribute;
 
 	USceneComponent* RootComponent = BoundObject->GetOwner()->GetRootComponent();
-	if (bNeedInverseMeshRotation && RootComponent && RootComponent != BoundObject)
+	if (RootComponent)
 	{
-		InverseMeshToActorRotation = RootComponent->GetComponentTransform().GetRelativeTransformReverse(BoundObject->GetComponentTransform()).GetRotation();
-	}
+		if (bNeedInverseMeshRotation && RootComponent != BoundObject)
+		{
+			InverseMeshToActorRotation = RootComponent->GetComponentTransform().GetRelativeTransformReverse(BoundObject->GetComponentTransform()).GetRotation();
+		}
 
-	ComponentTransform = BoundObject->GetRelativeTransform();
-	ActorTransform = RootComponent->GetRelativeTransform();
+		ActorTransform = RootComponent->GetRelativeTransform();
+	}
 
 	switch (RootDestination)
 	{
@@ -168,7 +172,7 @@ void FMovieSceneMixerRootMotionComponentData::Initialize()
 }
 
 FAnimNextConvertRootMotionToWorldSpaceTask::FAnimNextConvertRootMotionToWorldSpaceTask(const TSharedPtr<FMovieSceneMixerRootMotionComponentData>& InRootMotionData)
-	: RootMotionData(InRootMotionData)
+	: WeakRootMotionData(InRootMotionData)
 {
 }
 
@@ -195,6 +199,7 @@ void FAnimNextConvertRootMotionToWorldSpaceTask::Execute(UE::AnimNext::FEvaluati
 
 
 	FAnimMixerComponentTypes* AnimMixerComponents = FAnimMixerComponentTypes::Get();
+	TSharedPtr<FMovieSceneMixerRootMotionComponentData> RootMotionData = WeakRootMotionData.Pin();
 	if (RootMotionData)
 	{
 		if (FTransformAnimationAttribute* RootMotionAttribute = (*Keyframe)->Attributes.Find<FTransformAnimationAttribute>(AnimMixerComponents->RootTransformAttributeId))
@@ -207,27 +212,22 @@ void FAnimNextConvertRootMotionToWorldSpaceTask::Execute(UE::AnimNext::FEvaluati
 				RelativeTransform.SetTranslation(InverseMeshToActorRotation.GetValue() * RelativeTransform.GetTranslation());
 			}
 
-			if (FTransformAnimationAttribute* RootBaseAttribute = (*Keyframe)->Attributes.Find<FTransformAnimationAttribute>(AnimMixerComponents->RootTransformBaseAttributeId))
-			{
-				RootMotionAttribute->Value = RelativeTransform * RootBaseAttribute->Value;
-			}
-			else
-			{
-				RootMotionAttribute->Value = RelativeTransform;
-			}
+			RootMotionAttribute->Value = RelativeTransform * RootMotionData->ActorTransform;
 		}
 	}
 }
 
-FAnimNextStoreRootTransformTask::FAnimNextStoreRootTransformTask(const TSharedPtr<FMovieSceneMixerRootMotionComponentData>& InRootMotionData)
+FAnimNextStoreRootTransformTask::FAnimNextStoreRootTransformTask(const TSharedPtr<FMovieSceneMixerRootMotionComponentData>& InRootMotionData, bool bInComponentHasKeyedTransform, bool bInRootComponentHasKeyedTransform)
 	: RootMotionData(InRootMotionData)
+	, bComponentHasKeyedTransform(bInComponentHasKeyedTransform)
+	, bRootComponentHasKeyedTransform(bInRootComponentHasKeyedTransform)
 {
 
 }
 
-FAnimNextStoreRootTransformTask FAnimNextStoreRootTransformTask::Make(const TSharedPtr<FMovieSceneMixerRootMotionComponentData>& InRootMotionData)
+FAnimNextStoreRootTransformTask FAnimNextStoreRootTransformTask::Make(const TSharedPtr<FMovieSceneMixerRootMotionComponentData>& InRootMotionData, bool bInComponentHasKeyedTransform, bool bInRootComponentHasKeyedTransform)
 {
-	return FAnimNextStoreRootTransformTask(InRootMotionData);
+	return FAnimNextStoreRootTransformTask(InRootMotionData, bInComponentHasKeyedTransform, bInRootComponentHasKeyedTransform);
 }
 
 
@@ -250,8 +250,7 @@ void FAnimNextStoreRootTransformTask::Execute(UE::AnimNext::FEvaluationVM& VM) c
 	}
 
 	FTransformAnimationAttribute* RootMotionAttribute = (*Keyframe)->Attributes.Find<FTransformAnimationAttribute>(AnimMixerComponents->RootTransformAttributeId);
-	FTransformAnimationAttribute* RootBaseAttribute   = (*Keyframe)->Attributes.Find<FTransformAnimationAttribute>(AnimMixerComponents->RootTransformBaseAttributeId);
-	if (!RootMotionAttribute && !RootBaseAttribute)
+	if (!RootMotionAttribute)
 	{
 		return;
 	}
@@ -271,15 +270,13 @@ void FAnimNextStoreRootTransformTask::Execute(UE::AnimNext::FEvaluationVM& VM) c
 		}
 	}
 
-	FTransform RootTransform = RootMotionAttribute
-		? RootMotionAttribute->Value
-		: RootBaseAttribute->Value;
+	FTransform RootTransform = RootMotionAttribute->Value;
 
 	bool bRemoveRootMotionAttributes = true;
 
 	const int32 RootBoneIndex = (*Keyframe)->Pose.GetRefPose().GetLODBoneIndexFromSkeletonBoneIndex(0);
 
-	switch (RootMotionData->RootDestination)
+	switch (FinalDestination)
 	{
 	case EMovieSceneRootMotionDestination::Discard:
 		if (RootBoneIndex != INDEX_NONE)
@@ -291,6 +288,15 @@ void FAnimNextStoreRootTransformTask::Execute(UE::AnimNext::FEvaluationVM& VM) c
 	case EMovieSceneRootMotionDestination::RootBone:
 		if (RootBoneIndex != INDEX_NONE)
 		{
+			// If we're applying to the root bone
+			RootTransform = RootTransform * RootMotionData->ActorTransform.Inverse();
+
+			TOptional<FQuat> InverseMeshToActorRotation = RootMotionData->GetInverseMeshToActorRotation();
+			if (InverseMeshToActorRotation.IsSet())
+			{
+				RootTransform.SetTranslation(InverseMeshToActorRotation.GetValue().Inverse() * RootTransform.GetTranslation());
+			}
+
 			(*Keyframe)->Pose.LocalTransformsView[RootBoneIndex] = RootTransform;
 		}
 		return;
@@ -333,8 +339,8 @@ void FAnimNextStoreRootTransformTask::Execute(UE::AnimNext::FEvaluationVM& VM) c
 			}
 		}
 
-		FTransform ActualTransform = (RootMotionData->RootDestination == EMovieSceneRootMotionDestination::Component)
-			? RootMotionData->ComponentTransform + RootMotionDelta
+		FTransform ActualTransform = (FinalDestination == EMovieSceneRootMotionDestination::Component)
+			? RootMotionData->ComponentToWorldTransform.Inverse() + RootMotionDelta
 			: RootMotionData->ActorTransform + RootMotionDelta;
 
 		ActualTransform.NormalizeRotation();
@@ -352,20 +358,15 @@ void FAnimNextStoreRootTransformTask::Execute(UE::AnimNext::FEvaluationVM& VM) c
 		}
 	}
 
-	if (RootMotionData->RootDestination != EMovieSceneRootMotionDestination::Attribute)
+	if (FinalDestination != EMovieSceneRootMotionDestination::Attribute)
 	{
 		TWeakPtr<FMovieSceneMixerRootMotionComponentData> WeakRootMotion = RootMotionData;
 
-		auto ApplyRootMotion = [RootTransform, Destination = RootMotionData->RootDestination, WeakTarget = RootMotionData->Target]
+		auto ApplyRootMotion = [RootTransform, WeakTarget = RootMotionData->Target]
 		{
 			USceneComponent* TargetComponent = WeakTarget.Get();
 			if (TargetComponent && !GActorMovementTracker.IsActive())
 			{
-				if (Destination == EMovieSceneRootMotionDestination::Actor)
-				{
-					TargetComponent = TargetComponent->GetOwner()->GetRootComponent();
-				}
-
 				if (TargetComponent)
 				{
 					TargetComponent->SetRelativeLocationAndRotation(RootTransform.GetLocation(), RootTransform.GetRotation().Rotator());
@@ -407,6 +408,11 @@ UMovieSceneRootMotionSystem::UMovieSceneRootMotionSystem(const FObjectInitialize
 	}
 }
 
+bool UMovieSceneRootMotionSystem::IsTransformKeyed(const FObjectKey& Object) const
+{
+	return ObjectsWithTransforms.Contains(Object);
+}
+
 void UMovieSceneRootMotionSystem::OnLink()
 {
 	UMovieSceneAnimMixerSystem* AnimMixer = Linker->LinkSystem<UMovieSceneAnimMixerSystem>();
@@ -478,6 +484,25 @@ void UMovieSceneRootMotionSystem::OnRun(FSystemTaskPrerequisites& InPrerequisite
 			}
 		);
 	}
+
+	ObjectsWithTransforms.Reset();
+
+	// Gather which objects have transforms
+	FEntityTaskBuilder()
+	.Read(BuiltInComponents->BoundObjectKey)
+	.FilterAll({ TrackComponents->ComponentTransform.PropertyTag, BuiltInComponents->CustomPropertyIndex })
+	.Iterate_PerAllocation(&Linker->EntityManager,
+		[this](const FEntityAllocation* Allocation, const FObjectKey* RootMotionTargets)
+		{
+			const int32 Num = Allocation->Num();
+			for (int32 Index = 0; Index < Num; ++Index)
+			{
+				this->ObjectsWithTransforms.Add(RootMotionTargets[Index]);
+			}
+		}
+	);
+
+	ObjectsWithTransforms.Shrink();
 }
 
 void UMovieSceneRootMotionSystem::OnSchedulePersistentTasks(UE::MovieScene::IEntitySystemScheduler* TaskScheduler)
@@ -508,9 +533,10 @@ void UMovieSceneRootMotionSystem::OnSchedulePersistentTasks(UE::MovieScene::IEnt
 	.Schedule_PerEntity<FGatherRootDestinations>(&Linker->EntityManager, TaskScheduler);
 
 	// Reset root motion data and gather current component/actor transforms
-	FTaskID InitializeRootMotion = TaskScheduler->AddTask<FInitializeRootMotionTask>(
-		FTaskParams(TEXT("Initialize Root Motion")).ForceGameThread(),
-		AnimMixer);
+	FTaskParams InitialzeParams(TEXT("Initialize Root Motion"));
+	InitialzeParams.ForceGameThread();
+	InitialzeParams.bForcePropagateDownstream = true;
+	FTaskID InitializeRootMotion = TaskScheduler->AddTask<FInitializeRootMotionTask>(InitialzeParams, AnimMixer);
 
 	TaskScheduler->AddPrerequisite(WaitForAllTransforms,      InitializeRootMotion);
 	TaskScheduler->AddPrerequisite(GatherRootDestinationTask, InitializeRootMotion);
diff --git a/Engine/Plugins/Experimental/MovieSceneAnimMixer/Source/MovieSceneAnimMixer/Public/AnimMixerComponentTypes.h b/Engine/Plugins/Experimental/MovieSceneAnimMixer/Source/MovieSceneAnimMixer/Public/AnimMixerComponentTypes.h
--- a/Engine/Plugins/Experimental/MovieSceneAnimMixer/Source/MovieSceneAnimMixer/Public/AnimMixerComponentTypes.h
+++ b/Engine/Plugins/Experimental/MovieSceneAnimMixer/Source/MovieSceneAnimMixer/Public/AnimMixerComponentTypes.h
@@ -26,7 +26,6 @@ namespace UE::MovieScene
 
 		// A root motion transform in animation space to be applied (after space conversion)
 		MOVIESCENEANIMMIXER_API static const UE::Anim::FAttributeId RootTransformAttributeId;
-		MOVIESCENEANIMMIXER_API static const UE::Anim::FAttributeId RootTransformBaseAttributeId;
 		MOVIESCENEANIMMIXER_API static const UE::Anim::FAttributeId RootTransformWeightAttributeId;
 		MOVIESCENEANIMMIXER_API static const UE::Anim::FAttributeId LegacySwapRootBoneAttributeId;
 
diff --git a/Engine/Plugins/Experimental/MovieSceneAnimMixer/Source/MovieSceneAnimMixer/Public/Systems/MovieSceneRootMotionSystem.h b/Engine/Plugins/Experimental/MovieSceneAnimMixer/Source/MovieSceneAnimMixer/Public/Systems/MovieSceneRootMotionSystem.h
--- a/Engine/Plugins/Experimental/MovieSceneAnimMixer/Source/MovieSceneAnimMixer/Public/Systems/MovieSceneRootMotionSystem.h
+++ b/Engine/Plugins/Experimental/MovieSceneAnimMixer/Source/MovieSceneAnimMixer/Public/Systems/MovieSceneRootMotionSystem.h
@@ -4,6 +4,7 @@
 
 #include "EntitySystem/MovieSceneEntitySystem.h"
 #include "EntitySystem/MovieSceneEntityIDs.h"
+#include "MovieSceneAnimationMixerTrack.h"
 #include "AnimSequencerInstanceProxy.h"
 #include "EvaluationVM/EvaluationTask.h"
 #include "Tickable.h"
@@ -14,15 +15,6 @@ enum class EMovieSceneRootMotionDestination : uint8;
 class UMovieSceneAnimMixerSystem;
 class UMovieSceneRootMotionSystem;
 
-// Enum describing what space that root motion should be applied in.
-UENUM()
-enum class EMovieSceneRootMotionSpace : uint8
-{
-	// Root motion should be applied in animation space, meaning that it will be applied on top of the blended transform track and transform origin.
-	AnimationSpace,
-	// Root motion should be applied in world space, meaning that it will override any transform track or transform origin.
-	WorldSpace
-};
 
 USTRUCT()
 struct FSkeletalAnimationRootMotionOverride
@@ -59,8 +51,6 @@ struct FMovieSceneRootMotionSettings
 struct FMovieSceneMixerRootMotionComponentData
 {
 public:
-	TOptional<FTransform> GetSequenceAnimationTransform() const;
-	void SetSequenceAnimationTransform(const FTransform& InTransform);
 
 	TOptional<FQuat> GetInverseMeshToActorRotation() const;
 
@@ -79,11 +69,9 @@ public:
 	EMovieSceneRootMotionDestination RootDestination;
 
 	FTransform ActorTransform;
-	FTransform ComponentTransform;
+	FTransform ComponentToWorldTransform;
 
 private:
-	// Blend output result of transform track calculations
-	TOptional<FTransform> SequenceAnimationTransform;
 
 	// Optional inverse mesh component to actor rotation used to offset any mesh component rotation where applicable.
 	TOptional<FQuat> InverseMeshToActorRotation;
@@ -106,7 +94,7 @@ struct FAnimNextConvertRootMotionToWorldSpaceTask : public FAnimNextEvaluationTa
 
 	virtual void Execute(UE::AnimNext::FEvaluationVM& VM) const override;
 
-	TSharedPtr<FMovieSceneMixerRootMotionComponentData> RootMotionData;
+	TWeakPtr<FMovieSceneMixerRootMotionComponentData> WeakRootMotionData;
 };
 
 // Task that gets the final mixed root transform and stores it in the root motion data for later application.
@@ -120,13 +108,15 @@ struct FAnimNextStoreRootTransformTask : public FAnimNextEvaluationTask
 	DECLARE_ANIM_EVALUATION_TASK(FAnimNextStoreRootTransformTask)
 
 	FAnimNextStoreRootTransformTask() = default;
-	FAnimNextStoreRootTransformTask(const TSharedPtr<FMovieSceneMixerRootMotionComponentData>& InRootMotionData);
+	FAnimNextStoreRootTransformTask(const TSharedPtr<FMovieSceneMixerRootMotionComponentData>& InRootMotionData, bool bInComponentHasKeyedTransform, bool bInRootComponentHasKeyedTransform);
 
-	static FAnimNextStoreRootTransformTask Make(const TSharedPtr<FMovieSceneMixerRootMotionComponentData>& InRootMotionData);
+	static FAnimNextStoreRootTransformTask Make(const TSharedPtr<FMovieSceneMixerRootMotionComponentData>& InRootMotionData, bool bInComponentHasKeyedTransform, bool bInRootComponentHasKeyedTransform);
 
 	virtual void Execute(UE::AnimNext::FEvaluationVM& VM) const override;
 
 	TSharedPtr<FMovieSceneMixerRootMotionComponentData> RootMotionData;
+	bool bComponentHasKeyedTransform = false;
+	bool bRootComponentHasKeyedTransform = false;
 };
 
 
@@ -144,9 +134,13 @@ public:
 
 	UMovieSceneRootMotionSystem(const FObjectInitializer& ObjInit);
 
+	bool IsTransformKeyed(const FObjectKey& Object) const;
+
 private:
 
 	virtual void OnLink() override;
 	virtual void OnRun(FSystemTaskPrerequisites& InPrerequisites, FSystemSubsequentTasks& Subsequents) override;
 	virtual void OnSchedulePersistentTasks(UE::MovieScene::IEntitySystemScheduler* TaskScheduler) override;
+
+	TSet<FObjectKey> ObjectsWithTransforms;
 };
