Ensure that root motion information is cleaned up correctly in Sequencer, and that we correctly wait for root spaces to be evaluated before building the mixer

diff --git a/Engine/Plugins/Experimental/MovieSceneAnimMixer/Source/MovieSceneAnimMixer/Private/Systems/MovieSceneAnimMixerSystem.cpp b/Engine/Plugins/Experimental/MovieSceneAnimMixer/Source/MovieSceneAnimMixer/Private/Systems/MovieSceneAnimMixerSystem.cpp
--- a/Engine/Plugins/Experimental/MovieSceneAnimMixer/Source/MovieSceneAnimMixer/Private/Systems/MovieSceneAnimMixerSystem.cpp
+++ b/Engine/Plugins/Experimental/MovieSceneAnimMixer/Source/MovieSceneAnimMixer/Private/Systems/MovieSceneAnimMixerSystem.cpp
@@ -86,6 +86,13 @@ namespace UE::MovieScene
 				return;
 			}
 
+			if (Mixer->WeakEntries.Num() == 0)
+			{
+				MixerTask = nullptr;
+				Mixer->EvaluationProgram = nullptr;
+				return;
+			}
+
 			// Create a new eval program
 			Mixer->EvaluationProgram = MakeShared<UE::AnimNext::FEvaluationProgram>();
 
@@ -540,7 +547,6 @@ void UMovieSceneAnimMixerSystem::OnRun(FSystemTaskPrerequisites& InPrerequisites
 				.Add(AnimMixerComponents->MeshComponent, FObjectComponent::Weak(NewMixerKey.BoundObjectKey.ResolveObjectPtr()))
 				.Add(AnimMixerComponents->Target, NewMixerKey.Target)
 				.Add(AnimMixerComponents->MixerTask, nullptr)
-				.AddConditional(AnimMixerComponents->MixerRootMotion, NewMixer->WeakRootMotion.Pin(), NewMixer->WeakRootMotion.Pin().IsValid())
 				.AddTag(BuiltInComponents->Tags.RestoreState) // TODO: For now we always restore state on the mixer when it gets unlinked.
 				.CreateEntity(&Linker->EntityManager);
 		}
@@ -571,22 +577,12 @@ void UMovieSceneAnimMixerSystem::OnRun(FSystemTaskPrerequisites& InPrerequisites
 
 				Algo::Sort(Mixer->WeakEntries, SortPredicate);
 				Mixer->bNeedsResort = false;
-
-				if (Mixer->MixerEntityID)
-				{
-					if (TSharedPtr<FMovieSceneMixerRootMotionComponentData> RootMotion = Mixer->WeakRootMotion.Pin())
-					{
-						Linker->EntityManager.AddComponent(Mixer->MixerEntityID, AnimMixerComponents->MixerRootMotion, RootMotion);
-					}
-					else
-					{
-						Linker->EntityManager.RemoveComponent(Mixer->MixerEntityID, AnimMixerComponents->MixerRootMotion);
-					}
-				}
 			}
 
 			if (Mixer->WeakEntries.IsEmpty())
 			{
+				Mixer->EvaluationProgram = nullptr;
+
 				if (Mixer->MixerEntityID)
 				{
 					Linker->EntityManager.AddComponent(Mixer->MixerEntityID, BuiltInComponents->Tags.NeedsUnlink, EEntityRecursion::Full);
diff --git a/Engine/Plugins/Experimental/MovieSceneAnimMixer/Source/MovieSceneAnimMixer/Private/Systems/MovieSceneMixedSkeletalAnimationSystem.cpp b/Engine/Plugins/Experimental/MovieSceneAnimMixer/Source/MovieSceneAnimMixer/Private/Systems/MovieSceneMixedSkeletalAnimationSystem.cpp
--- a/Engine/Plugins/Experimental/MovieSceneAnimMixer/Source/MovieSceneAnimMixer/Private/Systems/MovieSceneMixedSkeletalAnimationSystem.cpp
+++ b/Engine/Plugins/Experimental/MovieSceneAnimMixer/Source/MovieSceneAnimMixer/Private/Systems/MovieSceneMixedSkeletalAnimationSystem.cpp
@@ -27,6 +27,7 @@
 #include "Component/AnimNextComponent.h"
 #include "Systems/MovieSceneComponentTransformSystem.h"
 #include "Systems/MovieSceneQuaternionInterpolationRotationSystem.h"
+#include "Systems/ByteChannelEvaluatorSystem.h"
 #include "EntitySystem/MovieSceneBoundObjectInstantiator.h"
 #include "EntitySystem/MovieSceneBoundSceneComponentInstantiator.h"
 #include "EntitySystem/MovieSceneEntityMutations.h"
@@ -481,6 +482,7 @@ UMovieSceneMixedSkeletalAnimationSystem::UMovieSceneMixedSkeletalAnimationSystem
 		FBuiltInComponentTypes* BuiltInComponents = FBuiltInComponentTypes::Get(); 
 
 		DefineImplicitPrerequisite(UMovieSceneEvalTimeSystem::StaticClass(), GetClass());
+		DefineImplicitPrerequisite(UByteChannelEvaluatorSystem::StaticClass(), GetClass());
 		DefineImplicitPrerequisite(UMovieSceneGenericBoundObjectInstantiator::StaticClass(), GetClass());
 		DefineImplicitPrerequisite(UMovieSceneBoundSceneComponentInstantiator::StaticClass(), GetClass());
 		DefineComponentConsumer(GetClass(), BuiltInComponents->BoundObject);
diff --git a/Engine/Plugins/Experimental/MovieSceneAnimMixer/Source/MovieSceneAnimMixer/Private/Systems/MovieSceneRootMotionSystem.cpp b/Engine/Plugins/Experimental/MovieSceneAnimMixer/Source/MovieSceneAnimMixer/Private/Systems/MovieSceneRootMotionSystem.cpp
--- a/Engine/Plugins/Experimental/MovieSceneAnimMixer/Source/MovieSceneAnimMixer/Private/Systems/MovieSceneRootMotionSystem.cpp
+++ b/Engine/Plugins/Experimental/MovieSceneAnimMixer/Source/MovieSceneAnimMixer/Private/Systems/MovieSceneRootMotionSystem.cpp
@@ -24,6 +24,7 @@
 #include "EntitySystem/Interrogation/MovieSceneInterrogationLinker.h"
 #include "Systems/MovieScenePiecewiseDoubleBlenderSystem.h"
 #include "Systems/DoubleChannelEvaluatorSystem.h"
+#include "Systems/ByteChannelEvaluatorSystem.h"
 #include "Misc/MemStack.h"
 #include "EntitySystem/MovieSceneEntitySystemTask.h"
 #include "EntitySystem/MovieSceneEntitySystemLinker.h"
@@ -218,7 +219,7 @@ void FAnimNextConvertRootMotionToWorldSpaceTask::Execute(UE::AnimNext::FEvaluati
 }
 
 FAnimNextStoreRootTransformTask::FAnimNextStoreRootTransformTask(const TSharedPtr<FMovieSceneMixerRootMotionComponentData>& InRootMotionData, bool bInComponentHasKeyedTransform, bool bInRootComponentHasKeyedTransform)
-	: RootMotionData(InRootMotionData)
+	: WeakRootMotionData(InRootMotionData)
 	, bComponentHasKeyedTransform(bInComponentHasKeyedTransform)
 	, bRootComponentHasKeyedTransform(bInRootComponentHasKeyedTransform)
 {
@@ -238,6 +239,7 @@ void FAnimNextStoreRootTransformTask::Execute(UE::AnimNext::FEvaluationVM& VM) c
 
 	FAnimMixerComponentTypes* AnimMixerComponents = FAnimMixerComponentTypes::Get();
 
+	TSharedPtr<FMovieSceneMixerRootMotionComponentData> RootMotionData = WeakRootMotionData.Pin();
 	if (!RootMotionData.IsValid() || !EnumHasAnyFlags(VM.GetFlags(), EEvaluationFlags::Attributes | EEvaluationFlags::Trajectory))
 	{
 		return;
@@ -402,6 +404,7 @@ UMovieSceneRootMotionSystem::UMovieSceneRootMotionSystem(const FObjectInitialize
 		// This must be run before the anim mixer to ensure that the anim mixer sets up its program correctly with root motion
 		DefineImplicitPrerequisite(GetClass(), UMovieSceneAnimMixerSystem::StaticClass());
 
+		DefineImplicitPrerequisite(UByteChannelEvaluatorSystem::StaticClass(), GetClass());
 		DefineImplicitPrerequisite(UDoubleChannelEvaluatorSystem::StaticClass(), GetClass());
 		DefineImplicitPrerequisite(UMovieScenePiecewiseDoubleBlenderSystem::StaticClass(), GetClass());
 		DefineImplicitPrerequisite(UMovieSceneComponentTransformSystem::StaticClass(), GetClass());
@@ -442,7 +445,7 @@ void UMovieSceneRootMotionSystem::OnRun(FSystemTaskPrerequisites& InPrerequisite
 		// Remove expiring root motions
 		FEntityTaskBuilder()
 		.Write(AnimMixerComponents->MixerRootMotion)
-		.FilterAll({ AnimMixerComponents->RootDestination, BuiltInComponents->Tags.NeedsUnlink })
+		.FilterAll({ BuiltInComponents->Tags.NeedsUnlink })
 		.Iterate_PerEntity( &Linker->EntityManager,
 			[](TSharedPtr<FMovieSceneMixerRootMotionComponentData>& OutRootMotion)
 			{
diff --git a/Engine/Plugins/Experimental/MovieSceneAnimMixer/Source/MovieSceneAnimMixer/Public/Systems/MovieSceneRootMotionSystem.h b/Engine/Plugins/Experimental/MovieSceneAnimMixer/Source/MovieSceneAnimMixer/Public/Systems/MovieSceneRootMotionSystem.h
--- a/Engine/Plugins/Experimental/MovieSceneAnimMixer/Source/MovieSceneAnimMixer/Public/Systems/MovieSceneRootMotionSystem.h
+++ b/Engine/Plugins/Experimental/MovieSceneAnimMixer/Source/MovieSceneAnimMixer/Public/Systems/MovieSceneRootMotionSystem.h
@@ -114,7 +114,7 @@ struct FAnimNextStoreRootTransformTask : public FAnimNextEvaluationTask
 
 	virtual void Execute(UE::AnimNext::FEvaluationVM& VM) const override;
 
-	TSharedPtr<FMovieSceneMixerRootMotionComponentData> RootMotionData;
+	TWeakPtr<FMovieSceneMixerRootMotionComponentData> WeakRootMotionData;
 	bool bComponentHasKeyedTransform = false;
 	bool bRootComponentHasKeyedTransform = false;
 };
