Sequencer: Remove actor transform when setting root motion to component transforms (same as root bone)

diff --git a/Engine/Plugins/Experimental/MovieSceneAnimMixer/Source/MovieSceneAnimMixer/Private/Systems/MovieSceneRootMotionSystem.cpp b/Engine/Plugins/Experimental/MovieSceneAnimMixer/Source/MovieSceneAnimMixer/Private/Systems/MovieSceneRootMotionSystem.cpp
--- a/Engine/Plugins/Experimental/MovieSceneAnimMixer/Source/MovieSceneAnimMixer/Private/Systems/MovieSceneRootMotionSystem.cpp
+++ b/Engine/Plugins/Experimental/MovieSceneAnimMixer/Source/MovieSceneAnimMixer/Private/Systems/MovieSceneRootMotionSystem.cpp
@@ -274,6 +274,19 @@ void FAnimNextStoreRootTransformTask::Execute(UE::AnimNext::FEvaluationVM& VM) c
 
 	FTransform RootTransform = RootMotionAttribute->Value;
 
+	auto WorldSpaceToComponentSpace = [RootMotionData](const FTransform& InTransform)
+	{
+		FTransform Result = InTransform * RootMotionData->ActorTransform.Inverse();
+
+		TOptional<FQuat> InverseMeshToActorRotation = RootMotionData->GetInverseMeshToActorRotation();
+		if (InverseMeshToActorRotation.IsSet())
+		{
+			Result.SetTranslation(InverseMeshToActorRotation.GetValue().Inverse() * Result.GetTranslation());
+		}
+
+		return Result;
+	};
+
 	bool bRemoveRootMotionAttributes = true;
 
 	const int32 RootBoneIndex = (*Keyframe)->Pose.GetRefPose().GetLODBoneIndexFromSkeletonBoneIndex(0);
@@ -290,16 +303,7 @@ void FAnimNextStoreRootTransformTask::Execute(UE::AnimNext::FEvaluationVM& VM) c
 	case EMovieSceneRootMotionDestination::RootBone:
 		if (RootBoneIndex != INDEX_NONE)
 		{
-			// If we're applying to the root bone
-			RootTransform = RootTransform * RootMotionData->ActorTransform.Inverse();
-
-			TOptional<FQuat> InverseMeshToActorRotation = RootMotionData->GetInverseMeshToActorRotation();
-			if (InverseMeshToActorRotation.IsSet())
-			{
-				RootTransform.SetTranslation(InverseMeshToActorRotation.GetValue().Inverse() * RootTransform.GetTranslation());
-			}
-
-			(*Keyframe)->Pose.LocalTransformsView[RootBoneIndex] = RootTransform;
+			(*Keyframe)->Pose.LocalTransformsView[RootBoneIndex] = WorldSpaceToComponentSpace(RootTransform);
 		}
 		return;
 
@@ -360,6 +364,11 @@ void FAnimNextStoreRootTransformTask::Execute(UE::AnimNext::FEvaluationVM& VM) c
 		}
 	}
 
+	if (FinalDestination == EMovieSceneRootMotionDestination::Component)
+	{
+		RootTransform = WorldSpaceToComponentSpace(RootTransform);
+	}
+
 	if (FinalDestination != EMovieSceneRootMotionDestination::Attribute)
 	{
 		TWeakPtr<FMovieSceneMixerRootMotionComponentData> WeakRootMotion = RootMotionData;
