Sequencer: Add pre-animated state tracking for root motion transforms applied to components and actors

diff --git a/Engine/Plugins/Experimental/MovieSceneAnimMixer/Source/MovieSceneAnimMixer/Private/Systems/MovieSceneMixedSkeletalAnimationSystem.cpp b/Engine/Plugins/Experimental/MovieSceneAnimMixer/Source/MovieSceneAnimMixer/Private/Systems/MovieSceneMixedSkeletalAnimationSystem.cpp
--- a/Engine/Plugins/Experimental/MovieSceneAnimMixer/Source/MovieSceneAnimMixer/Private/Systems/MovieSceneMixedSkeletalAnimationSystem.cpp
+++ b/Engine/Plugins/Experimental/MovieSceneAnimMixer/Source/MovieSceneAnimMixer/Private/Systems/MovieSceneMixedSkeletalAnimationSystem.cpp
@@ -22,6 +22,7 @@
 #include "Systems/WeightAndEasingEvaluatorSystem.h"
 #include "Systems/MovieSceneObjectPropertySystem.h"
 #include "EntitySystem/MovieScenePreAnimatedStateSystem.h"
+#include "PreAnimatedState/MovieScenePreAnimatedComponentTransformStorage.h"
 
 #include "Components/SkeletalMeshComponent.h"
 #include "Component/AnimNextComponent.h"
@@ -61,11 +62,13 @@ struct FSkeletalAnimMixerMutation : IMovieSceneEntityMutation
 	FBuiltInComponentTypes* BuiltInComponents;
 	FMovieSceneTracksComponentTypes* TrackComponents; 
 	FAnimMixerComponentTypes* AnimMixerComponents;
-	FSkeletalAnimMixerMutation()
+	UMovieSceneEntitySystemLinker* Linker = nullptr;
+	FSkeletalAnimMixerMutation(UMovieSceneEntitySystemLinker* InLinker)
 	{
 		BuiltInComponents = FBuiltInComponentTypes::Get();
 		TrackComponents = FMovieSceneTracksComponentTypes::Get();
 		AnimMixerComponents = FAnimMixerComponentTypes::Get();
+		Linker = InLinker;
 	}
 
 	virtual void CreateMutation(FEntityManager* EntityManager, FComponentMask* InOutEntityComponentTypes) const
@@ -86,6 +89,24 @@ struct FSkeletalAnimMixerMutation : IMovieSceneEntityMutation
 	{
 		const int32 Num = Allocation->Num();
 		FEntityAllocationWriteContext WriteContext = FEntityAllocationWriteContext::NewAllocation();
+
+		// @todo: figure out restore state semantics with root motion
+		const bool bWantsRestore = false;
+		const bool bCapturePreAnimatedState = Linker->PreAnimatedState.IsCapturingGlobalState() || bWantsRestore;
+		static const FName PreAnimatedTransformName("Transform");
+
+		FPreAnimatedEntityCaptureSource* EntityMetaData = nullptr;
+		TSharedPtr<FPreAnimatedComponentTransformStorage> ComponentTransformStorage;
+		if (bCapturePreAnimatedState)
+		{
+			EntityMetaData = Linker->PreAnimatedState.GetOrCreateEntityMetaData();
+			ComponentTransformStorage = Linker->PreAnimatedState.GetOrCreateStorage<FPreAnimatedComponentTransformStorage>();
+		}
+
+		TArrayView<const FMovieSceneEntityID> EntityIDs = Allocation->GetEntityIDs();
+		TComponentReader<FRootInstanceHandle> RootInstanceHandles = Allocation->ReadComponents(BuiltInComponents->RootInstanceHandle);
+		TComponentReader<UObject*> BoundObjects = Allocation->ReadComponents(BuiltInComponents->BoundObject);
+
 		TComponentReader<FMovieSceneSkeletalAnimationComponentData> SkeletalAnimationData = Allocation->ReadComponents(TrackComponents->SkeletalAnimation);
 		TComponentWriter<TInstancedStruct<FMovieSceneMixedAnimationTarget>> OutAnimTargets = Allocation->WriteComponents(AnimMixerComponents->Target, WriteContext);
 		TComponentWriter<TSharedPtr<FAnimNextEvaluationTask>> OutTasks = Allocation->WriteComponents(AnimMixerComponents->Task, WriteContext);
@@ -124,6 +145,27 @@ struct FSkeletalAnimMixerMutation : IMovieSceneEntityMutation
 				}
 				OutRootMotionSettings[Index].LegacySwapRootBone = LegacySwapRootBone;
 
+				if (bCapturePreAnimatedState)
+				{
+					USceneComponent* SwapRootComponent = nullptr;
+					if (USceneComponent* BoundComponent = Cast<USceneComponent>(BoundObjects[Index]))
+					{
+						switch (LegacySwapRootBone)
+						{
+							case ESwapRootBone::SwapRootBone_None: break;
+							case ESwapRootBone::SwapRootBone_Component: SwapRootComponent = BoundComponent; break;
+							case ESwapRootBone::SwapRootBone_Actor:     SwapRootComponent = BoundComponent->GetOwner()->GetRootComponent(); break;
+						}
+					}
+					if (SwapRootComponent)
+					{
+						FPreAnimatedStateEntry Entry = ComponentTransformStorage->MakeEntry(SwapRootComponent, PreAnimatedTransformName);
+						EntityMetaData->BeginTrackingEntity(Entry, EntityIDs[Index], RootInstanceHandles[Index], bWantsRestore);
+						ComponentTransformStorage->CachePreAnimatedTransform(FCachePreAnimatedValueParams(), SwapRootComponent);
+					}
+				}
+
+
 				if (Track)
 				{
 					if (Track->RootMotionParams.bRootMotionsDirty) 
@@ -518,7 +560,7 @@ void UMovieSceneMixedSkeletalAnimationSystem::OnRun(FSystemTaskPrerequisites& In
 		Filter.None({ BuiltInComponents->Tags.ImportedEntity });
 
 		// Initialize components for skeletal animation mixing
-		FSkeletalAnimMixerMutation SkeletalAnimMixerMutation;
+		FSkeletalAnimMixerMutation SkeletalAnimMixerMutation(Linker);
 		Linker->EntityManager.MutateAll(Filter, SkeletalAnimMixerMutation);
 	}
 }
diff --git a/Engine/Plugins/Experimental/MovieSceneAnimMixer/Source/MovieSceneAnimMixer/Private/Systems/MovieSceneRootMotionSystem.cpp b/Engine/Plugins/Experimental/MovieSceneAnimMixer/Source/MovieSceneAnimMixer/Private/Systems/MovieSceneRootMotionSystem.cpp
--- a/Engine/Plugins/Experimental/MovieSceneAnimMixer/Source/MovieSceneAnimMixer/Private/Systems/MovieSceneRootMotionSystem.cpp
+++ b/Engine/Plugins/Experimental/MovieSceneAnimMixer/Source/MovieSceneAnimMixer/Private/Systems/MovieSceneRootMotionSystem.cpp
@@ -30,6 +30,7 @@
 #include "EntitySystem/MovieSceneEntitySystemLinker.h"
 #include "Systems/MovieScenePropertyInstantiator.h"
 #include "Systems/MovieSceneComponentTransformSystem.h"
+#include "PreAnimatedState/MovieScenePreAnimatedComponentTransformStorage.h"
 #include "EntitySystem/Interrogation/MovieSceneInterrogatedPropertyInstantiator.h"
 #include "EntitySystem/MovieSceneEntitySystemRunner.h"
 #include "Misc/ScopeRWLock.h"
@@ -464,35 +465,95 @@ void UMovieSceneRootMotionSystem::OnRun(FSystemTaskPrerequisites& InPrerequisite
 
 		// Set up root motion behaviors
 		FEntityTaskBuilder()
+		.ReadEntityIDs()
+		.Read(BuiltInComponents->RootInstanceHandle)
 		.Read(BuiltInComponents->BoundObjectKey)
+		.Read(BuiltInComponents->BoundObject)
 		.Write(AnimMixerComponents->Target)
 		.Write(AnimMixerComponents->MixerRootMotion)
 		.FilterAll({ AnimMixerComponents->RootDestination, BuiltInComponents->Tags.NeedsLink })
-		.Iterate_PerEntity( &Linker->EntityManager,
-			[AnimMixer](FObjectKey BoundObjectKey, TInstancedStruct<FMovieSceneMixedAnimationTarget>& OutTarget, TSharedPtr<FMovieSceneMixerRootMotionComponentData>& OutRootMotion)
+		.Iterate_PerAllocation( &Linker->EntityManager,
+			[this, AnimMixer, BuiltInComponents](
+				FEntityAllocationIteratorItem Item,
+				TRead<FMovieSceneEntityID> EntityIDs,
+				TRead<FRootInstanceHandle> RootInstanceHandles,
+				TRead<FObjectKey> BoundObjectKeys,
+				TRead<UObject*> BoundObjects,
+				TWrite<TInstancedStruct<FMovieSceneMixedAnimationTarget>> OutTargets,
+				TWrite<TSharedPtr<FMovieSceneMixerRootMotionComponentData>> OutRootMotions)
 			{
-				// For new entities, we 'resolve' the animation target so if 'Automatic' is picked we choose the right one automatically.
-				OutTarget = UMovieSceneAnimMixerSystem::ResolveAnimationTarget(BoundObjectKey, OutTarget);
+				const FEntityAllocation* Allocation = Item.GetAllocation();
+				const int32              Num        = Allocation->Num();
 
-				FMovieSceneAnimMixerKey Key(BoundObjectKey, OutTarget);
 
-				if (OutRootMotion)
-				{
-					AnimMixer->AssignRootMotion(Key, OutRootMotion);
-				}
-				else if (TSharedPtr<FMovieSceneMixerRootMotionComponentData> RootMotion = AnimMixer->FindRootMotion(Key))
+				// @todo: figure out restore state semantics with root motion
+				const bool bWantsRestore = false; 
+				const bool bCapturePreAnimatedState = Linker->PreAnimatedState.IsCapturingGlobalState() || bWantsRestore;
+
+				FPreAnimatedEntityCaptureSource* EntityMetaData = nullptr;
+				TSharedPtr<FPreAnimatedComponentTransformStorage> ComponentTransformStorage;
+				if (bCapturePreAnimatedState)
 				{
-					OutRootMotion = RootMotion;
+					EntityMetaData = Linker->PreAnimatedState.GetOrCreateEntityMetaData();
+					ComponentTransformStorage = Linker->PreAnimatedState.GetOrCreateStorage<FPreAnimatedComponentTransformStorage>();
 				}
-				else
+
+				static const FName PreAnimatedTransformName("Transform");
+				for (int32 Index = 0; Index < Num; ++Index)
 				{
-					// Make a new one
-					OutRootMotion = MakeShared<FMovieSceneMixerRootMotionComponentData>();
-					AnimMixer->AssignRootMotion(Key, OutRootMotion);
+					USceneComponent* BoundObject = Cast<USceneComponent>(BoundObjects[Index]);
+					if (!BoundObject)
+					{
+						continue;
+					}
+
+					if (bCapturePreAnimatedState)
+					{
+						const FMovieSceneEntityID EntityID           = EntityIDs[Index];
+						const FRootInstanceHandle RootInstanceHandle = RootInstanceHandles[Index];
+						const FCachePreAnimatedValueParams CacheParams;
+
+						// Track transform for the component and the root component
+						FPreAnimatedStateEntry Entry = ComponentTransformStorage->MakeEntry(BoundObject, PreAnimatedTransformName);
+						EntityMetaData->BeginTrackingEntity(Entry, EntityID, RootInstanceHandle, bWantsRestore);
+						ComponentTransformStorage->CachePreAnimatedTransform(CacheParams, BoundObject);
+
+						USceneComponent* RootComponent = BoundObject->GetOwner()->GetRootComponent();
+						if (RootComponent && RootComponent != BoundObject)
+						{
+							FPreAnimatedStateEntry RootEntry = ComponentTransformStorage->MakeEntry(RootComponent, PreAnimatedTransformName);
+							EntityMetaData->BeginTrackingEntity(RootEntry, EntityID, RootInstanceHandle, bWantsRestore);
+							ComponentTransformStorage->CachePreAnimatedTransform(CacheParams, RootComponent);
+						}
+					}
+
+					FObjectKey BoundObjectKey = BoundObjectKeys[Index];
+					TInstancedStruct<FMovieSceneMixedAnimationTarget>& OutTarget = OutTargets[Index];
+					TSharedPtr<FMovieSceneMixerRootMotionComponentData>& OutRootMotion = OutRootMotions[Index];
+
+					// For new entities, we 'resolve' the animation target so if 'Automatic' is picked we choose the right one automatically.
+					OutTarget = UMovieSceneAnimMixerSystem::ResolveAnimationTarget(BoundObjectKey, OutTarget);
+
+					FMovieSceneAnimMixerKey Key(BoundObjectKey, OutTarget);
+
+					if (OutRootMotion)
+					{
+						AnimMixer->AssignRootMotion(Key, OutRootMotion);
+					}
+					else if (TSharedPtr<FMovieSceneMixerRootMotionComponentData> RootMotion = AnimMixer->FindRootMotion(Key))
+					{
+						OutRootMotion = RootMotion;
+					}
+					else
+					{
+						// Make a new one
+						OutRootMotion = MakeShared<FMovieSceneMixerRootMotionComponentData>();
+						AnimMixer->AssignRootMotion(Key, OutRootMotion);
+					}
+
+					USceneComponent* Target = Cast<USceneComponent>(BoundObjectKey.ResolveObjectPtr());
+					OutRootMotion->OriginalBoundObject = Target;
 				}
-
-				USceneComponent* Target = Cast<USceneComponent>(BoundObjectKey.ResolveObjectPtr());
-				OutRootMotion->OriginalBoundObject = Target;
 			}
 		);
 	}
diff --git a/Engine/Source/Runtime/MovieSceneTracks/Private/PreAnimatedState/MovieScenePreAnimatedComponentTransformStorage.cpp b/Engine/Source/Runtime/MovieSceneTracks/Private/PreAnimatedState/MovieScenePreAnimatedComponentTransformStorage.cpp
--- a/Engine/Source/Runtime/MovieSceneTracks/Private/PreAnimatedState/MovieScenePreAnimatedComponentTransformStorage.cpp
+++ b/Engine/Source/Runtime/MovieSceneTracks/Private/PreAnimatedState/MovieScenePreAnimatedComponentTransformStorage.cpp
@@ -84,63 +84,69 @@ FPreAnimatedComponentTransformStorage::FPreAnimatedComponentTransformStorage()
 
 void FPreAnimatedComponentTransformStorage::CachePreAnimatedTransforms(const FCachePreAnimatedValueParams& Params, TArrayView<UObject* const> BoundObjects, TOptional<TFunctionRef<bool(int32)>> Predicate)
 {
-	static FMovieScenePropertyBinding PropertyBinding("Transform", TEXT("Transform"));
 
 	for (int32 Index = 0; Index < BoundObjects.Num(); ++Index)
 	{
 		UObject* BoundObject = BoundObjects[Index];
-		if (!BoundObject || (Predicate && !(*Predicate)(Index)))
+		if (BoundObject && (!Predicate || (*Predicate)(Index)))
 		{
-			continue;
+			CachePreAnimatedTransform(Params, BoundObject);
 		}
+	}
+}
 
-		TOptional<FResolvedProperty> Property = FPropertyRegistry::ResolveProperty(BoundObject, PropertyBinding, CustomAccessors);
-		if (!Property)
-		{
-			continue;
-		}
+void FPreAnimatedComponentTransformStorage::CachePreAnimatedTransform(const FCachePreAnimatedValueParams& Params, UObject* BoundObject)
+{
+	check(BoundObject);
+
+	static FMovieScenePropertyBinding PropertyBinding("Transform", TEXT("Transform"));
+
+	TOptional<FResolvedProperty> Property = FPropertyRegistry::ResolveProperty(BoundObject, PropertyBinding, CustomAccessors);
+	if (!Property)
+	{
+		return;
+	}
+
+	TTuple<FObjectKey, FName> Key{ BoundObject, PropertyBinding.PropertyPath };
 
-		TTuple<FObjectKey, FName> Key{ BoundObject, PropertyBinding.PropertyPath };
+	FPreAnimatedStorageGroupHandle GroupHandle  = this->Traits.MakeGroup(BoundObject);
+	FPreAnimatedStorageIndex       StorageIndex = this->GetOrCreateStorageIndex(Key);
 
-		FPreAnimatedStorageGroupHandle GroupHandle  = this->Traits.MakeGroup(BoundObject);
-		FPreAnimatedStorageIndex       StorageIndex = this->GetOrCreateStorageIndex(Key);
+	FPreAnimatedStateEntry Entry{ GroupHandle, FPreAnimatedStateCachedValueHandle{ StorageID, StorageIndex } };
 
-		FPreAnimatedStateEntry Entry{ GroupHandle, FPreAnimatedStateCachedValueHandle{ StorageID, StorageIndex } };
+	ParentExtension->EnsureMetaData(Entry);
 
-		ParentExtension->EnsureMetaData(Entry);
+	EPreAnimatedStorageRequirement StorageRequirement = ParentExtension->GetStorageRequirement(Entry);
+	if (!this->IsStorageRequirementSatisfied(StorageIndex, StorageRequirement))
+	{
+		StorageType NewValue;
 
-		EPreAnimatedStorageRequirement StorageRequirement = ParentExtension->GetStorageRequirement(Entry);
-		if (!this->IsStorageRequirementSatisfied(StorageIndex, StorageRequirement))
+		if (const uint16* Fast = Property->TryGet<uint16>())
 		{
-			StorageType NewValue;
-
-			if (const uint16* Fast = Property->TryGet<uint16>())
-			{
-				NewValue.Binding.template Set<uint16>(*Fast);
-				FComponentTransformPropertyTraits::GetObjectPropertyValue(BoundObject, *Fast, NewValue.Data);
-			}
-			else if (const FCustomPropertyIndex* CustomIndex = Property->TryGet<FCustomPropertyIndex>())
-			{
-				const FCustomPropertyAccessor& Accessor = CustomAccessors[CustomIndex->Value];
-
-				NewValue.Binding.template Set<const FCustomPropertyAccessor*>(&Accessor);
-				FComponentTransformPropertyTraits::GetObjectPropertyValue(BoundObject, Accessor, NewValue.Data);
-			}
-			else
-			{
-				const TSharedPtr<FTrackInstancePropertyBindings>& Bindings = Property->Get<TSharedPtr<FTrackInstancePropertyBindings>>();
-
-				NewValue.Binding.template Set<TSharedPtr<FTrackInstancePropertyBindings>>(Bindings);
-				FComponentTransformPropertyTraits::GetObjectPropertyValue(BoundObject, Bindings.Get(), NewValue.Data);
-			}
-
-			this->AssignPreAnimatedValue(StorageIndex, StorageRequirement, MoveTemp(NewValue));
+			NewValue.Binding.template Set<uint16>(*Fast);
+			FComponentTransformPropertyTraits::GetObjectPropertyValue(BoundObject, *Fast, NewValue.Data);
 		}
+		else if (const FCustomPropertyIndex* CustomIndex = Property->TryGet<FCustomPropertyIndex>())
+		{
+			const FCustomPropertyAccessor& Accessor = CustomAccessors[CustomIndex->Value];
 
-		if (Params.bForcePersist)
+			NewValue.Binding.template Set<const FCustomPropertyAccessor*>(&Accessor);
+			FComponentTransformPropertyTraits::GetObjectPropertyValue(BoundObject, Accessor, NewValue.Data);
+		}
+		else
 		{
-			this->ForciblyPersistStorage(StorageIndex);
+			const TSharedPtr<FTrackInstancePropertyBindings>& Bindings = Property->Get<TSharedPtr<FTrackInstancePropertyBindings>>();
+
+			NewValue.Binding.template Set<TSharedPtr<FTrackInstancePropertyBindings>>(Bindings);
+			FComponentTransformPropertyTraits::GetObjectPropertyValue(BoundObject, Bindings.Get(), NewValue.Data);
 		}
+
+		this->AssignPreAnimatedValue(StorageIndex, StorageRequirement, MoveTemp(NewValue));
+	}
+
+	if (Params.bForcePersist)
+	{
+		this->ForciblyPersistStorage(StorageIndex);
 	}
 }
 
diff --git a/Engine/Source/Runtime/MovieSceneTracks/Private/PreAnimatedState/MovieScenePreAnimatedComponentTransformStorage.h b/Engine/Source/Runtime/MovieSceneTracks/Private/PreAnimatedState/MovieScenePreAnimatedComponentTransformStorage.h
deleted file mode 100644
diff --git a/Engine/Source/Runtime/MovieSceneTracks/Public/PreAnimatedState/MovieScenePreAnimatedComponentTransformStorage.h b/Engine/Source/Runtime/MovieSceneTracks/Public/PreAnimatedState/MovieScenePreAnimatedComponentTransformStorage.h
new file mode 100644
--- /dev/null
+++ b/Engine/Source/Runtime/MovieSceneTracks/Public/PreAnimatedState/MovieScenePreAnimatedComponentTransformStorage.h
@@ -0,0 +1,39 @@
+// Copyright Epic Games, Inc. All Rights Reserved.
+
+#pragma once
+
+#include "CoreTypes.h"
+#include "MovieSceneTracksPropertyTypes.h"
+#include "MovieSceneTracksComponentTypes.h"
+#include "Evaluation/PreAnimatedState/MovieScenePreAnimatedPropertyStorage.h"
+
+namespace UE
+{
+namespace MovieScene
+{
+
+struct MOVIESCENETRACKS_API FComponentTransformPreAnimatedTraits : FComponentTransformPropertyTraits
+{
+	using KeyType     = FObjectKey;
+	using StorageType = FIntermediate3DTransform;
+
+	/** These override the functions in FComponentTransformPropertyTraits in order to wrap them with a temporary mobility assignment */
+	static void SetObjectPropertyValue(UObject* InObject, const FCustomPropertyAccessor& BaseCustomAccessor, const FIntermediate3DTransform& CachedTransform);
+	static void SetObjectPropertyValue(UObject* InObject, uint16 PropertyOffset, const FIntermediate3DTransform& CachedTransform);
+	static void SetObjectPropertyValue(UObject* InObject, FTrackInstancePropertyBindings* PropertyBindings, const FIntermediate3DTransform& CachedTransform);
+};
+
+struct FPreAnimatedComponentTransformStorage
+	: TPreAnimatedPropertyStorage<FComponentTransformPreAnimatedTraits>
+{
+	MOVIESCENETRACKS_API FPreAnimatedComponentTransformStorage();
+
+	static MOVIESCENETRACKS_API TAutoRegisterPreAnimatedStorageID<FPreAnimatedComponentTransformStorage> StorageID;
+
+	MOVIESCENETRACKS_API void CachePreAnimatedTransform(const FCachePreAnimatedValueParams& Params, UObject* BoundObject);
+	MOVIESCENETRACKS_API void CachePreAnimatedTransforms(const FCachePreAnimatedValueParams& Params, TArrayView<UObject* const> BoundObjects, TOptional<TFunctionRef<bool(int32)>> Predicate = TOptional<TFunctionRef<bool(int32)>>());
+};
+
+
+} // namespace MovieScene
+} // namespace UE
