Fix Android compile errors

diff --git a/Engine/Plugins/Experimental/MovieSceneAnimMixer/Source/MovieSceneAnimMixer/Private/Systems/MovieSceneRootMotionSystem.cpp b/Engine/Plugins/Experimental/MovieSceneAnimMixer/Source/MovieSceneAnimMixer/Private/Systems/MovieSceneRootMotionSystem.cpp
--- a/Engine/Plugins/Experimental/MovieSceneAnimMixer/Source/MovieSceneAnimMixer/Private/Systems/MovieSceneRootMotionSystem.cpp
+++ b/Engine/Plugins/Experimental/MovieSceneAnimMixer/Source/MovieSceneAnimMixer/Private/Systems/MovieSceneRootMotionSystem.cpp
@@ -361,7 +361,7 @@ void FAnimNextStoreRootTransformTask::Execute(UE::AnimNext::FEvaluationVM& VM) c
 		}
 		else if (!bRemoveRootMotionAttributes)
 		{
-			(*Keyframe)->Attributes.Add(AnimMixerComponents->RootTransformAttributeId, FTransformAnimationAttribute(RootTransform));
+			(*Keyframe)->Attributes.Add(AnimMixerComponents->RootTransformAttributeId, FTransformAnimationAttribute{ RootTransform });
 		}
 	}
 
@@ -534,7 +534,7 @@ void UMovieSceneRootMotionSystem::OnRun(FSystemTaskPrerequisites& InPrerequisite
 					// For new entities, we 'resolve' the animation target so if 'Automatic' is picked we choose the right one automatically.
 					OutTarget = UMovieSceneAnimMixerSystem::ResolveAnimationTarget(BoundObjectKey, OutTarget);
 
-					FMovieSceneAnimMixerKey Key(BoundObjectKey, OutTarget);
+					FMovieSceneAnimMixerKey Key{BoundObjectKey, OutTarget};
 
 					if (OutRootMotion)
 					{
