Sequencer: Apply root bone offsets after root motion when root motion is in world space

It doesn't make much sense to apply a pre-root motion transform when in world space

diff --git a/Engine/Plugins/Experimental/MovieSceneAnimMixer/Source/MovieSceneAnimMixer/Private/Systems/MovieSceneMixedSkeletalAnimationSystem.cpp b/Engine/Plugins/Experimental/MovieSceneAnimMixer/Source/MovieSceneAnimMixer/Private/Systems/MovieSceneMixedSkeletalAnimationSystem.cpp
--- a/Engine/Plugins/Experimental/MovieSceneAnimMixer/Source/MovieSceneAnimMixer/Private/Systems/MovieSceneMixedSkeletalAnimationSystem.cpp
+++ b/Engine/Plugins/Experimental/MovieSceneAnimMixer/Source/MovieSceneAnimMixer/Private/Systems/MovieSceneMixedSkeletalAnimationSystem.cpp
@@ -245,6 +245,7 @@ struct FGatherMixableSkeletalAnimations
 		TReadOptional<FMovieSceneRootMotionSettings> OptionalRootMotionSettings,
 		const double* OptLocationX, const double* OptLocationY, const double* OptLocationZ,
 		const double* OptRotationX, const double* OptRotationY, const double* OptRotationZ,
+		const uint8* OptSpaces,
 		TWrite<TSharedPtr<FAnimNextEvaluationTask>> EvalTask) const
 	{
 		// Gather all the skeletal animations currently active in all sequences.
@@ -377,6 +378,11 @@ struct FGatherMixableSkeletalAnimations
 
 			AnimTask->RootBaseTransform.SetRotation(CurrentRotation.Quaternion());
 			AnimTask->RootBaseTransform.SetLocation(CurrentLocation);
+
+			if (OptSpaces)
+			{
+				AnimTask->RootMotionSpace = (EMovieSceneRootMotionSpace)(OptSpaces[Index]);
+			}
 		}
 	}
 
@@ -488,8 +494,15 @@ void FMovieSceneSkeletalAnimationEvaluationTask::Execute(UE::AnimNext::FEvaluati
 			FTransformAnimationAttribute* RootMotionAttribute        = Keyframe.Attributes.FindOrAdd<FTransformAnimationAttribute>(AnimMixerComponents->RootTransformAttributeId);
 			FFloatAnimationAttribute*     RootMotionWeight           = Keyframe.Attributes.FindOrAdd<FFloatAnimationAttribute>(AnimMixerComponents->RootTransformWeightAttributeId);
 
-			RootMotionAttribute->Value        = RootBoneTransform.GetValue() * RootBaseTransform;
-			RootMotionWeight->Value           = 1.f;
+			if (RootMotionSpace == EMovieSceneRootMotionSpace::AnimationSpace)
+			{
+				RootMotionAttribute->Value = RootBoneTransform.GetValue() * RootBaseTransform;
+			}
+			else
+			{
+				RootMotionAttribute->Value = RootBaseTransform * RootBoneTransform.GetValue();
+			}
+			RootMotionWeight->Value = 1.f;
 
 			if (AnimationData.LegacySwapRootBone != ESwapRootBone::SwapRootBone_None)
 			{
@@ -597,6 +610,7 @@ void UMovieSceneMixedSkeletalAnimationSystem::OnSchedulePersistentTasks(UE::Movi
 		.ReadOptional(BuiltInComponents->DoubleResult[3]) // Root Rotation Offset X
 		.ReadOptional(BuiltInComponents->DoubleResult[4]) // Root Rotation Offset Y
 		.ReadOptional(BuiltInComponents->DoubleResult[5]) // Root Rotation Offset Z
+		.ReadOptional(BuiltInComponents->ByteResult)      // Root Space
 		.Write(AnimMixerComponents->Task)
 		.FilterAll( { TrackComponents->Tags.AnimMixerPoseProducer })
 		.FilterNone({ BuiltInComponents->Tags.Ignored})
diff --git a/Engine/Plugins/Experimental/MovieSceneAnimMixer/Source/MovieSceneAnimMixer/Public/Systems/MovieSceneMixedSkeletalAnimationSystem.h b/Engine/Plugins/Experimental/MovieSceneAnimMixer/Source/MovieSceneAnimMixer/Public/Systems/MovieSceneMixedSkeletalAnimationSystem.h
--- a/Engine/Plugins/Experimental/MovieSceneAnimMixer/Source/MovieSceneAnimMixer/Public/Systems/MovieSceneMixedSkeletalAnimationSystem.h
+++ b/Engine/Plugins/Experimental/MovieSceneAnimMixer/Source/MovieSceneAnimMixer/Public/Systems/MovieSceneMixedSkeletalAnimationSystem.h
@@ -58,6 +58,9 @@ struct FMovieSceneSkeletalAnimationEvaluationTask : public FAnimNextEvaluationTa
 
 	UPROPERTY()
 	FTransform RootBaseTransform;
+
+	UPROPERTY()
+	EMovieSceneRootMotionSpace RootMotionSpace = EMovieSceneRootMotionSpace::AnimationSpace;
 };
 
 // System to handle creating evaluation tasks from skeletal animation track sections for the anim mixer.
