Fixed Legacy swap root bone settings misbehaving during blended regions

diff --git a/Engine/Plugins/Experimental/MovieSceneAnimMixer/Source/MovieSceneAnimMixer/Private/AnimMixerComponentTypes.cpp b/Engine/Plugins/Experimental/MovieSceneAnimMixer/Source/MovieSceneAnimMixer/Private/AnimMixerComponentTypes.cpp
--- a/Engine/Plugins/Experimental/MovieSceneAnimMixer/Source/MovieSceneAnimMixer/Private/AnimMixerComponentTypes.cpp
+++ b/Engine/Plugins/Experimental/MovieSceneAnimMixer/Source/MovieSceneAnimMixer/Private/AnimMixerComponentTypes.cpp
@@ -12,7 +12,6 @@ namespace UE::MovieScene
 
 	const UE::Anim::FAttributeId FAnimMixerComponentTypes::RootTransformAttributeId(TEXT("RootTransform"), FCompactPoseBoneIndex(0));
 	const UE::Anim::FAttributeId FAnimMixerComponentTypes::RootTransformWeightAttributeId(TEXT("RootTransformWeight"), FCompactPoseBoneIndex(0));
-	const UE::Anim::FAttributeId FAnimMixerComponentTypes::LegacySwapRootBoneAttributeId(TEXT("LegacySwapRootBone"), FCompactPoseBoneIndex(0));
 
 	static TUniquePtr<FAnimMixerComponentTypes> GAnimMixerComponentTypes;
 	static bool GAnimMixerComponentTypesDestroyed = false;
diff --git a/Engine/Plugins/Experimental/MovieSceneAnimMixer/Source/MovieSceneAnimMixer/Private/Systems/MovieSceneMixedSkeletalAnimationSystem.cpp b/Engine/Plugins/Experimental/MovieSceneAnimMixer/Source/MovieSceneAnimMixer/Private/Systems/MovieSceneMixedSkeletalAnimationSystem.cpp
--- a/Engine/Plugins/Experimental/MovieSceneAnimMixer/Source/MovieSceneAnimMixer/Private/Systems/MovieSceneMixedSkeletalAnimationSystem.cpp
+++ b/Engine/Plugins/Experimental/MovieSceneAnimMixer/Source/MovieSceneAnimMixer/Private/Systems/MovieSceneMixedSkeletalAnimationSystem.cpp
@@ -503,12 +503,6 @@ void FMovieSceneSkeletalAnimationEvaluationTask::Execute(UE::AnimNext::FEvaluati
 				RootMotionAttribute->Value = RootBaseTransform * RootBoneTransform.GetValue();
 			}
 			RootMotionWeight->Value = 1.f;
-
-			if (AnimationData.LegacySwapRootBone != ESwapRootBone::SwapRootBone_None)
-			{
-				FIntegerAnimationAttribute* Int = Keyframe.Attributes.FindOrAdd<FIntegerAnimationAttribute>(AnimMixerComponents->LegacySwapRootBoneAttributeId);
-				Int->Value = (int32)AnimationData.LegacySwapRootBone;
-			}
 		}
 
 		VM.PushValue(KEYFRAME_STACK_NAME, MakeUnique<FKeyframeState>(MoveTemp(Keyframe)));
diff --git a/Engine/Plugins/Experimental/MovieSceneAnimMixer/Source/MovieSceneAnimMixer/Private/Systems/MovieSceneRootMotionSystem.cpp b/Engine/Plugins/Experimental/MovieSceneAnimMixer/Source/MovieSceneAnimMixer/Private/Systems/MovieSceneRootMotionSystem.cpp
--- a/Engine/Plugins/Experimental/MovieSceneAnimMixer/Source/MovieSceneAnimMixer/Private/Systems/MovieSceneRootMotionSystem.cpp
+++ b/Engine/Plugins/Experimental/MovieSceneAnimMixer/Source/MovieSceneAnimMixer/Private/Systems/MovieSceneRootMotionSystem.cpp
@@ -260,19 +260,6 @@ void FAnimNextStoreRootTransformTask::Execute(UE::AnimNext::FEvaluationVM& VM) c
 
 	EMovieSceneRootMotionDestination FinalDestination = RootMotionData->RootDestination;
 
-	FIntegerAnimationAttribute* Int = (*Keyframe)->Attributes.Find<FIntegerAnimationAttribute>(AnimMixerComponents->LegacySwapRootBoneAttributeId);
-	if (Int)
-	{
-		if (Int->Value == (int32)ESwapRootBone::SwapRootBone_Actor)
-		{
-			FinalDestination = EMovieSceneRootMotionDestination::Actor;
-		}
-		else if (Int->Value == (int32)ESwapRootBone::SwapRootBone_Component)
-		{
-			FinalDestination = EMovieSceneRootMotionDestination::Component;
-		}
-	}
-
 	FTransform RootTransform = RootMotionAttribute->Value;
 
 	auto WorldSpaceToComponentSpace = [RootMotionData](const FTransform& InTransform)
diff --git a/Engine/Plugins/Experimental/MovieSceneAnimMixer/Source/MovieSceneAnimMixer/Public/AnimMixerComponentTypes.h b/Engine/Plugins/Experimental/MovieSceneAnimMixer/Source/MovieSceneAnimMixer/Public/AnimMixerComponentTypes.h
--- a/Engine/Plugins/Experimental/MovieSceneAnimMixer/Source/MovieSceneAnimMixer/Public/AnimMixerComponentTypes.h
+++ b/Engine/Plugins/Experimental/MovieSceneAnimMixer/Source/MovieSceneAnimMixer/Public/AnimMixerComponentTypes.h
@@ -27,7 +27,6 @@ namespace UE::MovieScene
 		// A root motion transform in animation space to be applied (after space conversion)
 		MOVIESCENEANIMMIXER_API static const UE::Anim::FAttributeId RootTransformAttributeId;
 		MOVIESCENEANIMMIXER_API static const UE::Anim::FAttributeId RootTransformWeightAttributeId;
-		MOVIESCENEANIMMIXER_API static const UE::Anim::FAttributeId LegacySwapRootBoneAttributeId;
 
 		TComponentTypeID<int32> Priority;
 		TComponentTypeID<TInstancedStruct<FMovieSceneMixedAnimationTarget>> Target;
