Fixed potential for a null ptr crash in anim next traget initialization / restoration

diff --git a/Engine/Plugins/Experimental/MovieSceneAnimMixer/Source/MovieSceneAnimMixer/Private/Systems/MovieSceneAnimNextTargetSystem.cpp b/Engine/Plugins/Experimental/MovieSceneAnimMixer/Source/MovieSceneAnimMixer/Private/Systems/MovieSceneAnimNextTargetSystem.cpp
--- a/Engine/Plugins/Experimental/MovieSceneAnimMixer/Source/MovieSceneAnimMixer/Private/Systems/MovieSceneAnimNextTargetSystem.cpp
+++ b/Engine/Plugins/Experimental/MovieSceneAnimMixer/Source/MovieSceneAnimMixer/Private/Systems/MovieSceneAnimNextTargetSystem.cpp
@@ -102,7 +102,10 @@ namespace UE::MovieScene
 
 			FSkeletalMeshEditorParams State;
 			State.SaveState(SkelMeshComp);
-			SkelMeshComp->GetSceneProxy()->SetCanSkipRedundantTransformUpdates(false);
+			if (FPrimitiveSceneProxy* Proxy = SkelMeshComp->GetSceneProxy())
+			{
+				Proxy->SetCanSkipRedundantTransformUpdates(false);
+			}
 			return State;
 		}
 
@@ -111,7 +114,10 @@ namespace UE::MovieScene
 			USkeletalMeshComponent* SkelMeshComp = Cast<USkeletalMeshComponent>(Object.ResolveObjectPtr());
 			if (SkelMeshComp != nullptr && SkelMeshComp->GetSceneProxy())
 			{
-				SkelMeshComp->GetSceneProxy()->SetCanSkipRedundantTransformUpdates(true);
+				if (FPrimitiveSceneProxy* Proxy = SkelMeshComp->GetSceneProxy())
+				{
+					Proxy->SetCanSkipRedundantTransformUpdates(true);
+				}
 				InOutCachedValue.RestoreState();
 			}
 		}
