Sequencer: Fixed root motion being applied to component transforms when the component itelf is a root component

diff --git a/Engine/Plugins/Experimental/MovieSceneAnimMixer/Source/MovieSceneAnimMixer/Private/Systems/MovieSceneRootMotionSystem.cpp b/Engine/Plugins/Experimental/MovieSceneAnimMixer/Source/MovieSceneAnimMixer/Private/Systems/MovieSceneRootMotionSystem.cpp
--- a/Engine/Plugins/Experimental/MovieSceneAnimMixer/Source/MovieSceneAnimMixer/Private/Systems/MovieSceneRootMotionSystem.cpp
+++ b/Engine/Plugins/Experimental/MovieSceneAnimMixer/Source/MovieSceneAnimMixer/Private/Systems/MovieSceneRootMotionSystem.cpp
@@ -129,7 +129,6 @@ void FMovieSceneMixerRootMotionComponentData::Initialize()
 	InverseMeshToActorRotation.Reset();
 
 	USceneComponent* BoundObject = OriginalBoundObject.Get();
-	USceneComponent* NewTarget = nullptr;
 
 	if (!BoundObject)
 	{
@@ -171,6 +170,11 @@ void FMovieSceneMixerRootMotionComponentData::Initialize()
 		Target = RootComponent;
 		break;
 	}
+
+	// If we want to swap the root bone with the component transform, but that component is the root component,
+	//     that is the same behavior as swapping with the actor, so don't perform any inverse component transformations
+	bComponentSpaceRoot = (RootDestination == EMovieSceneRootMotionDestination::RootBone)
+		|| (RootDestination == EMovieSceneRootMotionDestination::Component && RootComponent != Target);
 }
 
 FAnimNextConvertRootMotionToWorldSpaceTask::FAnimNextConvertRootMotionToWorldSpaceTask(const TSharedPtr<FMovieSceneMixerRootMotionComponentData>& InRootMotionData, const FTransform& InTransformOrigin, const FTransform& InBaseTransform, ESpaceConversions InConversions)
@@ -287,22 +291,20 @@ void FAnimNextStoreRootTransformTask::Execute(UE::AnimNext::FEvaluationVM& VM) c
 		return;
 	}
 
-	EMovieSceneRootMotionDestination FinalDestination = RootMotionData->RootDestination;
+	const EMovieSceneRootMotionDestination FinalDestination = RootMotionData->RootDestination;
 
 	FTransform RootTransform = RootMotionAttribute->Value;
 
-	auto WorldSpaceToComponentSpace = [RootMotionData](const FTransform& InTransform)
+	if (RootMotionData->bComponentSpaceRoot)
 	{
-		FTransform Result = InTransform * RootMotionData->ActorTransform.Inverse();
+		RootTransform *= RootMotionData->ActorTransform.Inverse();
 
 		TOptional<FQuat> InverseMeshToActorRotation = RootMotionData->GetInverseMeshToActorRotation();
 		if (InverseMeshToActorRotation.IsSet())
 		{
-			Result.SetTranslation(InverseMeshToActorRotation.GetValue().Inverse() * Result.GetTranslation());
+			RootTransform.SetTranslation(InverseMeshToActorRotation.GetValue().Inverse() * RootTransform.GetTranslation());
 		}
-
-		return Result;
-	};
+	}
 
 	bool bRemoveRootMotionAttributes = true;
 
@@ -320,7 +322,7 @@ void FAnimNextStoreRootTransformTask::Execute(UE::AnimNext::FEvaluationVM& VM) c
 	case EMovieSceneRootMotionDestination::RootBone:
 		if (RootBoneIndex != INDEX_NONE)
 		{
-			(*Keyframe)->Pose.LocalTransformsView[RootBoneIndex] = WorldSpaceToComponentSpace(RootTransform);
+			(*Keyframe)->Pose.LocalTransformsView[RootBoneIndex] = RootTransform;
 		}
 		return;
 
@@ -381,11 +383,6 @@ void FAnimNextStoreRootTransformTask::Execute(UE::AnimNext::FEvaluationVM& VM) c
 		}
 	}
 
-	if (FinalDestination == EMovieSceneRootMotionDestination::Component)
-	{
-		RootTransform = WorldSpaceToComponentSpace(RootTransform);
-	}
-
 	if (FinalDestination != EMovieSceneRootMotionDestination::Attribute)
 	{
 		TWeakPtr<FMovieSceneMixerRootMotionComponentData> WeakRootMotion = RootMotionData;
@@ -395,10 +392,7 @@ void FAnimNextStoreRootTransformTask::Execute(UE::AnimNext::FEvaluationVM& VM) c
 			USceneComponent* TargetComponent = WeakTarget.Get();
 			if (TargetComponent && !GActorMovementTracker.IsActive())
 			{
-				if (TargetComponent)
-				{
-					TargetComponent->SetRelativeLocationAndRotation(RootTransform.GetLocation(), RootTransform.GetRotation().Rotator());
-				}
+				TargetComponent->SetRelativeLocationAndRotation(RootTransform.GetLocation(), RootTransform.GetRotation().Rotator());
 			}
 		};
 		FAnimNextModuleInstance::RunTaskOnGameThread(MoveTemp(ApplyRootMotion));
diff --git a/Engine/Plugins/Experimental/MovieSceneAnimMixer/Source/MovieSceneAnimMixer/Public/Systems/MovieSceneRootMotionSystem.h b/Engine/Plugins/Experimental/MovieSceneAnimMixer/Source/MovieSceneAnimMixer/Public/Systems/MovieSceneRootMotionSystem.h
--- a/Engine/Plugins/Experimental/MovieSceneAnimMixer/Source/MovieSceneAnimMixer/Public/Systems/MovieSceneRootMotionSystem.h
+++ b/Engine/Plugins/Experimental/MovieSceneAnimMixer/Source/MovieSceneAnimMixer/Public/Systems/MovieSceneRootMotionSystem.h
@@ -77,6 +77,10 @@ private:
 	TOptional<FQuat> InverseMeshToActorRotation;
 
 	mutable FTransactionallySafeRWLock RootMotionLock;
+
+public:
+
+	bool bComponentSpaceRoot = false;
 };
 
 // Task that converts the root motion attribute on the top pose of the pose stack to world space by adding on the actor transformation, root base transform, and/or transform origin.
