Make MovieSceneAnimNextMixer enable anim in editor for current / child mesh components.

diff --git a/Engine/Plugins/Experimental/MovieSceneAnimMixer/Source/MovieSceneAnimMixer/Private/Systems/MovieSceneAnimNextTargetSystem.cpp b/Engine/Plugins/Experimental/MovieSceneAnimMixer/Source/MovieSceneAnimMixer/Private/Systems/MovieSceneAnimNextTargetSystem.cpp
--- a/Engine/Plugins/Experimental/MovieSceneAnimMixer/Source/MovieSceneAnimMixer/Private/Systems/MovieSceneAnimNextTargetSystem.cpp
+++ b/Engine/Plugins/Experimental/MovieSceneAnimMixer/Source/MovieSceneAnimMixer/Private/Systems/MovieSceneAnimNextTargetSystem.cpp
@@ -94,13 +94,13 @@ namespace UE::MovieScene
 	struct FPreAnimatedSkelMeshComponentTraits : FBoundObjectPreAnimatedStateTraits
 	{
 		using KeyType = FObjectKey;
-		using StorageType = FSkeletalMeshEditorParams;
+		using StorageType = FSkeletalMeshRestoreState;
 
-		FSkeletalMeshEditorParams CachePreAnimatedValue(const KeyType& Object)
+		StorageType CachePreAnimatedValue(const KeyType& Object)
 		{
 			USkeletalMeshComponent* SkelMeshComp = CastChecked<USkeletalMeshComponent>(Object.ResolveObjectPtr());
 
-			FSkeletalMeshEditorParams State;
+			StorageType State;
 			State.SaveState(SkelMeshComp);
 			if (FPrimitiveSceneProxy* Proxy = SkelMeshComp->GetSceneProxy())
 			{
@@ -109,7 +109,7 @@ namespace UE::MovieScene
 			return State;
 		}
 
-		void RestorePreAnimatedValue(const KeyType& Object, FSkeletalMeshEditorParams& InOutCachedValue, const FRestoreStateParams& Params)
+		void RestorePreAnimatedValue(const KeyType& Object, StorageType& InOutCachedValue, const FRestoreStateParams& Params)
 		{
 			USkeletalMeshComponent* SkelMeshComp = Cast<USkeletalMeshComponent>(Object.ResolveObjectPtr());
 			if (SkelMeshComp != nullptr && SkelMeshComp->GetSceneProxy())
