Anim Mixer: Character transitions from cinematics into gamplay now work seamlessly when using the anim mixer

diff --git a/Engine/Plugins/Experimental/Mover/Source/Mover/Private/DefaultMovementSet/LayeredMoves/RootMotionAttributeLayeredMove.cpp b/Engine/Plugins/Experimental/Mover/Source/Mover/Private/DefaultMovementSet/LayeredMoves/RootMotionAttributeLayeredMove.cpp
--- a/Engine/Plugins/Experimental/Mover/Source/Mover/Private/DefaultMovementSet/LayeredMoves/RootMotionAttributeLayeredMove.cpp
+++ b/Engine/Plugins/Experimental/Mover/Source/Mover/Private/DefaultMovementSet/LayeredMoves/RootMotionAttributeLayeredMove.cpp
@@ -31,10 +31,6 @@ FAutoConsoleVariable CVarDisableRootMotionAttrContributions(
 static const FName RootMotionAttributeName = "RootMotionDelta";
 static const UE::Anim::FAttributeId RootMotionAttributeId = { RootMotionAttributeName, FCompactPoseBoneIndex(0) };
 
-
-static const FName RootTransformAttributeName = "RootTransform";
-static const UE::Anim::FAttributeId RootTransformAttributeId = { RootTransformAttributeName, FCompactPoseBoneIndex(0) };
-
 FLayeredMove_RootMotionAttribute::FLayeredMove_RootMotionAttribute()
 {
 	DurationMs = -1.f;
@@ -52,8 +48,6 @@ bool FLayeredMove_RootMotionAttribute::GenerateMove(const FMoverTickStartData& S
 
 	const float DeltaSeconds = TimeStep.StepMs / 1000.f;
 
-	bool bDidAttrHaveRootTransform = false;
-	FTransform RootTransform;
 	bool bDidAttrHaveRootMotion = false;
 	FTransform LocalRootMotion;
 	bool bHasValidWarpingContext = false;
@@ -72,12 +66,7 @@ bool FLayeredMove_RootMotionAttribute::GenerateMove(const FMoverTickStartData& S
 		if (Mesh)
 		{
 			// TODO: support options for different interpretations, such as velocity
-			if (const FTransformAnimationAttribute* RootTransformAttribute = Mesh->GetCustomAttributes().Find<FTransformAnimationAttribute>(RootTransformAttributeId))
-			{
-				RootTransform = RootTransformAttribute->Value;
-				bDidAttrHaveRootTransform = true;
-			}
-			else if (const FTransformAnimationAttribute* RootMotionAttribute = Mesh->GetCustomAttributes().Find<FTransformAnimationAttribute>(RootMotionAttributeId))
+			if (const FTransformAnimationAttribute* RootMotionAttribute = Mesh->GetCustomAttributes().Find<FTransformAnimationAttribute>(RootMotionAttributeId))
 			{
 				// NOTE this will only work for ticking modes that tick in time with the world tick, because it relies on the Mesh ticking at the same rate as the movement simulation. 
 				// For fixed-tick modes, the attribute would be better as an accumulator for the movement sim to consume, along with a time accumulation attribute.
@@ -130,22 +119,7 @@ bool FLayeredMove_RootMotionAttribute::GenerateMove(const FMoverTickStartData& S
 		WarpingContext			= WarpingContextForResim;
 	}
 
-	if (bDidAttrHaveRootTransform)
-	{
-		const FMoverDefaultSyncState* SyncState = SimState.SyncState.SyncStateCollection.FindDataByType<FMoverDefaultSyncState>();
-
-		const FTransform SimActorTransform = FTransform(SyncState->GetOrientation_WorldSpace().Quaternion(), SyncState->GetLocation_WorldSpace());
-
-		OutProposedMove = FProposedMove();
-		OutProposedMove.MixMode = EMoveMixMode::OverrideAll;
-
-		// Convert the transform into linear and angular velocities
-		OutProposedMove.LinearVelocity  = (RootTransform.GetTranslation() - SimActorTransform.GetTranslation()) / DeltaSeconds;
-		OutProposedMove.AngularVelocity = ((RootTransform.GetRotation()-SimActorTransform.GetRotation()) / DeltaSeconds).Rotator();
-
-		return true;
-	}
-	else if (bDidAttrHaveRootMotion)
+	if (bDidAttrHaveRootMotion)
 	{
 		const FMoverDefaultSyncState* SyncState = SimState.SyncState.SyncStateCollection.FindDataByType<FMoverDefaultSyncState>();
 		const FCharacterDefaultInputs* InputCmd = SimState.InputCmd.InputCollection.FindDataByType<FCharacterDefaultInputs>();
diff --git a/Engine/Plugins/Experimental/MovieSceneAnimMixer/Source/MovieSceneAnimMixer/Private/Systems/MovieSceneAnimMixerSystem.cpp b/Engine/Plugins/Experimental/MovieSceneAnimMixer/Source/MovieSceneAnimMixer/Private/Systems/MovieSceneAnimMixerSystem.cpp
--- a/Engine/Plugins/Experimental/MovieSceneAnimMixer/Source/MovieSceneAnimMixer/Private/Systems/MovieSceneAnimMixerSystem.cpp
+++ b/Engine/Plugins/Experimental/MovieSceneAnimMixer/Source/MovieSceneAnimMixer/Private/Systems/MovieSceneAnimMixerSystem.cpp
@@ -27,6 +27,8 @@
 #include "EvaluationVM/EvaluationVM.h"
 #include "EvaluationVM/KeyframeState.h"
 #include "EvaluationVM/EvaluationProgram.h"
+#include "Animation/BuiltInAttributeTypes.h"
+#include "Animation/AnimRootMotionProvider.h"
 #include "Algo/ForEach.h"
 #include "Algo/MaxElement.h"
 #include "Systems/MovieScenePropertyInstantiator.h"
@@ -279,7 +281,7 @@ namespace UE::MovieScene
 						}
 						else
 						{
-							Mixer->EvaluationProgram->AppendTask(FAnimNextBlendTwoKeyframesTask::Make(MixerEntry->PoseWeight));
+							Mixer->EvaluationProgram->AppendTask(FAnimNextBlendTwoKeyframesPreserveRootMotionTask::Make(MixerEntry->PoseWeight));
 						}
 					}
 				}
@@ -291,7 +293,7 @@ namespace UE::MovieScene
 					// For now we always blend weighted averages with the next pose with a weight of 1.
 					//    Ultimately this should be controlled by its own weight within an anim mixer layer track.
 					constexpr float WeightedAverageBlendWeight = 1.f;
-					Mixer->EvaluationProgram->AppendTask(FAnimNextBlendTwoKeyframesTask::Make(WeightedAverageBlendWeight));
+					Mixer->EvaluationProgram->AppendTask(FAnimNextBlendTwoKeyframesPreserveRootMotionTask::Make(WeightedAverageBlendWeight));
 				}
 			}
 
@@ -304,6 +306,110 @@ namespace UE::MovieScene
 	};
 }
 
+FAnimNextBlendTwoKeyframesPreserveRootMotionTask FAnimNextBlendTwoKeyframesPreserveRootMotionTask::Make(float InterpolationAlpha)
+{
+	FAnimNextBlendTwoKeyframesPreserveRootMotionTask Task;
+	Task.InterpolationAlpha = InterpolationAlpha;
+	return Task;
+}
+
+void FAnimNextBlendTwoKeyframesPreserveRootMotionTask::Execute(UE::AnimNext::FEvaluationVM& VM) const
+{
+	using namespace UE::AnimNext;
+	using namespace UE::MovieScene;
+
+	const UE::Anim::FAttributeId RootMotionDeltaAttributeId(UE::Anim::IAnimRootMotionProvider::AttributeName, FCompactPoseBoneIndex(0));
+
+	const FAnimMixerComponentTypes* AnimMixerComponents = FAnimMixerComponentTypes::Get();
+
+	TOptional<FTransform> OverrideRootTransform;
+	TOptional<float>      OverrideRootWeight;
+	TOptional<FTransform> OverrideRootDeltaTransform;
+
+	// This task is used to perform a two way blend between 2 poses without blending root motion attributes
+	//    if they do not exist in either pose. By default these attributes would blend with the identity matrix
+	//    in a way that would cause undesirable effects when blending sequencer's pose with the incoming pose
+	//    from the upstream graph.
+
+	//        [ Sequencer pose A           ]  [ Sequencer poseB            ]
+	//        |    + RootMotionTransform   |  |                            |
+	//        [____________________________]  [____________________________]
+	//                                 \           /
+	//                                  \  Blend  /
+	//                                   \       /
+	//                                    \     /
+	//                                     \   /
+	//                                      \ /
+	//                       [ Final Sequencer Pose       ]           [ Upstream pose (locomotion) ]
+	//                       |  + RootMotionTransform     |           |    + RootMotionDelta       |
+	//                       [____________________________]           [____________________________]
+	//                                                    \           /
+	//                                                     \  Blend  /
+	//                                                      \       /
+	//                                                       \     /
+	//                                                        \   /
+	//                                                         \ /
+	//                                           [ Final Pose               ]
+	//                                           |    + RootMotionTransform |
+	//                                           |    + RootMotionDelta     |
+	//                                           [__________________________]
+	// From this final pose, FAnimNextStoreRootTransformTask is able to accurately read both Sequencer's desired world space
+	//    root transform, its desired weight, and the incoming desired root motion delta from locomotion. It can then blend
+	//    all these things together to form the final root motion delta which can be consumed by external systems like Mover
+
+	{
+		const TUniquePtr<FKeyframeState>* KeyframeA = VM.PeekValue<TUniquePtr<FKeyframeState>>(KEYFRAME_STACK_NAME, 0);
+		const TUniquePtr<FKeyframeState>* KeyframeB = VM.PeekValue<TUniquePtr<FKeyframeState>>(KEYFRAME_STACK_NAME, 1);
+		if (KeyframeA && KeyframeB)
+		{
+			UE::Anim::FStackAttributeContainer& AttributesA = (*KeyframeA)->Attributes;
+			UE::Anim::FStackAttributeContainer& AttributesB = (*KeyframeB)->Attributes;
+
+			const FTransformAnimationAttribute* RootMotionTransformA = AttributesA.Find<FTransformAnimationAttribute>(AnimMixerComponents->RootTransformAttributeId);
+			const FTransformAnimationAttribute* RootMotionDeltaA     = AttributesA.Find<FTransformAnimationAttribute>(RootMotionDeltaAttributeId);
+
+			const FTransformAnimationAttribute* RootMotionTransformB = AttributesB.Find<FTransformAnimationAttribute>(AnimMixerComponents->RootTransformAttributeId);
+			const FTransformAnimationAttribute* RootMotionDeltaB     = AttributesB.Find<FTransformAnimationAttribute>(RootMotionDeltaAttributeId);
+
+			if (!RootMotionDeltaA && RootMotionDeltaB)
+			{
+				// Preserve root motion *delta* from B if it's not in A (ie, don't blend it with identity!)
+				OverrideRootDeltaTransform = RootMotionDeltaB->Value;
+			}
+
+			if (RootMotionTransformA && !RootMotionTransformB)
+			{
+				// Preserve root motion transform from A if it's not in B (ie, don't blend it with identity!)
+				OverrideRootTransform = RootMotionTransformA->Value;
+			}
+		}
+	}
+
+	// Do the actual blend
+	FAnimNextBlendTwoKeyframesTask::Execute(VM);
+
+	if (OverrideRootTransform || OverrideRootDeltaTransform)
+	{
+		TUniquePtr<FKeyframeState> Keyframe;
+		if (VM.PopValue(KEYFRAME_STACK_NAME, Keyframe) && Keyframe)
+		{
+			if (OverrideRootTransform)
+			{
+				Keyframe->Attributes.FindOrAdd<FTransformAnimationAttribute>(AnimMixerComponents->RootTransformAttributeId)->Value
+					= OverrideRootTransform.GetValue();
+			}
+			if (OverrideRootDeltaTransform)
+			{
+				Keyframe->Attributes.FindOrAdd<FTransformAnimationAttribute>(RootMotionDeltaAttributeId)->Value
+					= OverrideRootDeltaTransform.GetValue();
+			}
+
+			VM.PushValue(KEYFRAME_STACK_NAME, MoveTemp(Keyframe));
+		}
+	}
+}
+
+
 FMovieSceneAccumulateAbsoluteBlendTask FMovieSceneAccumulateAbsoluteBlendTask::Make(float ScaleFactor)
 {
 	FMovieSceneAccumulateAbsoluteBlendTask Task;
diff --git a/Engine/Plugins/Experimental/MovieSceneAnimMixer/Source/MovieSceneAnimMixer/Private/Systems/MovieSceneMixedSkeletalAnimationSystem.cpp b/Engine/Plugins/Experimental/MovieSceneAnimMixer/Source/MovieSceneAnimMixer/Private/Systems/MovieSceneMixedSkeletalAnimationSystem.cpp
--- a/Engine/Plugins/Experimental/MovieSceneAnimMixer/Source/MovieSceneAnimMixer/Private/Systems/MovieSceneMixedSkeletalAnimationSystem.cpp
+++ b/Engine/Plugins/Experimental/MovieSceneAnimMixer/Source/MovieSceneAnimMixer/Private/Systems/MovieSceneMixedSkeletalAnimationSystem.cpp
@@ -461,17 +461,11 @@ void FMovieSceneSkeletalAnimationEvaluationTask::Execute(UE::AnimNext::FEvaluati
 		// Trajectory is currently held as an attribute
 		if (EnumHasAnyFlags(VM.GetFlags(), EEvaluationFlags::Attributes | EEvaluationFlags::Trajectory))
 		{
-			// If the sequence has root motion enabled, allow sampling of a root motion delta into the custom attribute container of the outgoing pose
-			if (AnimSequencePtr->HasRootMotion())
-			{
-				// TODO: We should cache the provider in the VM
-				// We have to grab two locks to get it and it won't change during graph evaluation
-				if (const UE::Anim::IAnimRootMotionProvider* RootMotionProvider = UE::Anim::IAnimRootMotionProvider::Get())
-				{
-					RootMotionProvider->SampleRootMotion(ExtractionContext.DeltaTimeRecord, *AnimSequencePtr, ExtractionContext.bLooping, Keyframe.Attributes);
-					// TODO: Do I need to override this too?
-				}
-			}
+			// We intentionally do not use IAnimRootMotionProvider here to sample root motion:
+			//     We don't want to emit and blend RootMotionDelta attributes for every pose in the mixer,
+			//     we want to blend absolute RootMotionTransform attributes, and emit the final RootMotionDelta
+			//     attribute once, after all blending has been performed. This allows us to blend seamlessly with
+			//     world-space transforms from gamplay
 
 			// If we haven't already grabbed the root transform, grab it here in case we're not doing the full bone sample
 			if (!RootBoneTransform.IsSet())
diff --git a/Engine/Plugins/Experimental/MovieSceneAnimMixer/Source/MovieSceneAnimMixer/Private/Systems/MovieSceneRootMotionSystem.cpp b/Engine/Plugins/Experimental/MovieSceneAnimMixer/Source/MovieSceneAnimMixer/Private/Systems/MovieSceneRootMotionSystem.cpp
--- a/Engine/Plugins/Experimental/MovieSceneAnimMixer/Source/MovieSceneAnimMixer/Private/Systems/MovieSceneRootMotionSystem.cpp
+++ b/Engine/Plugins/Experimental/MovieSceneAnimMixer/Source/MovieSceneAnimMixer/Private/Systems/MovieSceneRootMotionSystem.cpp
@@ -34,6 +34,7 @@
 #include "EntitySystem/Interrogation/MovieSceneInterrogatedPropertyInstantiator.h"
 #include "EntitySystem/MovieSceneEntitySystemRunner.h"
 #include "Misc/ScopeRWLock.h"
+#include "VisualLogger/VisualLogger.h"
 
 #if WITH_EDITOR
 #include "Editor.h"
@@ -306,8 +307,6 @@ void FAnimNextStoreRootTransformTask::Execute(UE::AnimNext::FEvaluationVM& VM) c
 		}
 	}
 
-	bool bRemoveRootMotionAttributes = true;
-
 	const int32 RootBoneIndex = (*Keyframe)->Pose.GetRefPose().GetLODBoneIndexFromSkeletonBoneIndex(0);
 
 	switch (FinalDestination)
@@ -327,13 +326,8 @@ void FAnimNextStoreRootTransformTask::Execute(UE::AnimNext::FEvaluationVM& VM) c
 		return;
 
 	case EMovieSceneRootMotionDestination::Component:
-		break;
-
 	case EMovieSceneRootMotionDestination::Actor:
-		break;
-
 	case EMovieSceneRootMotionDestination::Attribute:
-		bRemoveRootMotionAttributes = false;
 		break;
 	}
 
@@ -341,45 +335,61 @@ void FAnimNextStoreRootTransformTask::Execute(UE::AnimNext::FEvaluationVM& VM) c
 	const FFloatAnimationAttribute* RootMotionWeight = (*Keyframe)->Attributes.Find<FFloatAnimationAttribute>(AnimMixerComponents->RootTransformWeightAttributeId);
 	const float TransformWeight = RootMotionWeight ? RootMotionWeight->Value : 1.f;
 
-	if (TransformWeight < 1.f)
-	{
-		static const UE::Anim::FAttributeId RootMotionDeltaAttributeId(UE::Anim::IAnimRootMotionProvider::AttributeName , FCompactPoseBoneIndex(0));
-
-		FTransformAnimationAttribute* RootMotionDeltaAttribute = (*Keyframe)->Attributes.Find<FTransformAnimationAttribute>(RootMotionDeltaAttributeId);
-
-		// Undo the weight of the root motion
-		RootTransform *= ScalarRegister(1.f / TransformWeight);
-		RootTransform.NormalizeRotation();
 
+	// If we have a root motion provider, convert our prospective root motion transform into a delta so that
+	//   pose history and locomotion react correctly. This code also supports blending into and out of gameplay
+	//   animation
+	if (const UE::Anim::IAnimRootMotionProvider* RootMotionProvider = UE::Anim::IAnimRootMotionProvider::Get())
+	{
 		FTransform RootMotionDelta;
-		if (RootMotionDeltaAttribute)
+		if (RootMotionProvider->ExtractRootMotion((*Keyframe)->Attributes, RootMotionDelta))
 		{
-			// Factor in the current root motion delta in order to support blending to/from gameplay-root-motion
-			RootMotionDelta = RootMotionDeltaAttribute->Value;
+			const FTransform ActorHeadingTransform = RootMotionData->ActorTransform;
 
 			TOptional<FQuat> InverseMeshToActorRotation = RootMotionData->GetInverseMeshToActorRotation();
 			if (InverseMeshToActorRotation.IsSet())
 			{
+				// Rotate the translation of the root motion delta to compensate for component rotation
 				RootMotionDelta.SetTranslation(InverseMeshToActorRotation.GetValue() * RootMotionDelta.GetTranslation());
 			}
-		}
 
-		FTransform ActualTransform = (FinalDestination == EMovieSceneRootMotionDestination::Component)
-			? RootMotionData->ComponentToWorldTransform.Inverse() + RootMotionDelta
-			: RootMotionData->ActorTransform + RootMotionDelta;
+			// Blend the root motion delta with Sequencer's desired delta based on its weight
+			const FTransform LocomotionDelta  = RootMotionDelta;
+			const FTransform DesiredRootDelta = RootTransform * ActorHeadingTransform.Inverse();
 
-		ActualTransform.NormalizeRotation();
-		ActualTransform.BlendWith(RootTransform, TransformWeight);
+			FTransform BlendResult(FMath::Lerp(RootMotionDelta.GetTranslation(), DesiredRootDelta.GetTranslation(), TransformWeight));
+			BlendResult.SetRotation(FQuat::Slerp(RootMotionDelta.GetRotation(), DesiredRootDelta.GetRotation(), TransformWeight));
 
-		RootTransform = ActualTransform;
+			// Assign the final result in component space
+			RootMotionDelta = BlendResult;
+			if (InverseMeshToActorRotation.IsSet())
+			{
+				// Unrotate the translation of the root motion delta to compensate for component rotation
+				RootMotionDelta.SetTranslation(InverseMeshToActorRotation->Inverse() * RootMotionDelta.GetTranslation());
+			}
 
-		if (RootMotionAttribute)
-		{
-			RootMotionAttribute->Value = RootTransform;
-		}
-		else if (!bRemoveRootMotionAttributes)
-		{
-			(*Keyframe)->Attributes.Add(AnimMixerComponents->RootTransformAttributeId, FTransformAnimationAttribute{ RootTransform });
+#if ENABLE_VISUAL_LOG
+			if (FVisualLogger::IsRecording())
+			{
+				static const TCHAR* LogName = TEXT("MovieSceneRootMotion");
+				USceneComponent* Component = RootMotionData->OriginalBoundObject.Get();
+
+				auto DrawMarker = [Component](const FTransform& Transform, const FColor& Color, auto&& Name)
+				{
+					FVector Dir(50.f, 0.f, 0.f);
+
+					UE_VLOG_CIRCLE(Component, LogName, Display, Transform.GetLocation(), FVector::UpVector, 10.f, Color, Name);
+					UE_VLOG_ARROW(Component, LogName, Display, Transform.GetLocation(), Transform.GetLocation() + Transform.GetRotation()*Dir, Color, TEXT(""));
+				};
+
+				DrawMarker(ActorHeadingTransform,                   FColorList::Black, TEXT("Actor"));
+				DrawMarker(RootTransform,                           FColorList::Blue,  TEXT("Sequencer"));
+				DrawMarker(LocomotionDelta * ActorHeadingTransform, FColorList::Red,   TEXT("Locomotion"));
+				DrawMarker(BlendResult     * ActorHeadingTransform, FColorList::Green, TEXT("Result"));
+			}
+#endif
+
+			RootMotionProvider->OverrideRootMotion(RootMotionDelta, (*Keyframe)->Attributes);
 		}
 	}
 
@@ -397,14 +407,6 @@ void FAnimNextStoreRootTransformTask::Execute(UE::AnimNext::FEvaluationVM& VM) c
 		};
 		FAnimNextModuleInstance::RunTaskOnGameThread(MoveTemp(ApplyRootMotion));
 	}
-
-
-	// We're done with these attributes now
-	if (bRemoveRootMotionAttributes)
-	{
-		(*Keyframe)->Attributes.Remove<FTransformAnimationAttribute>(AnimMixerComponents->RootTransformAttributeId);
-		(*Keyframe)->Attributes.Remove<FFloatAnimationAttribute>(AnimMixerComponents->RootTransformWeightAttributeId);
-	}
 }
 
 UMovieSceneRootMotionSystem::UMovieSceneRootMotionSystem(const FObjectInitializer& ObjInit)
diff --git a/Engine/Plugins/Experimental/MovieSceneAnimMixer/Source/MovieSceneAnimMixer/Public/Systems/MovieSceneAnimMixerSystem.h b/Engine/Plugins/Experimental/MovieSceneAnimMixer/Source/MovieSceneAnimMixer/Public/Systems/MovieSceneAnimMixerSystem.h
--- a/Engine/Plugins/Experimental/MovieSceneAnimMixer/Source/MovieSceneAnimMixer/Public/Systems/MovieSceneAnimMixerSystem.h
+++ b/Engine/Plugins/Experimental/MovieSceneAnimMixer/Source/MovieSceneAnimMixer/Public/Systems/MovieSceneAnimMixerSystem.h
@@ -10,6 +10,7 @@
 #include "Templates/SharedPointer.h"
 #include "AnimSequencerInstanceProxy.h"
 #include "Systems/MovieSceneRootMotionSystem.h"
+#include "EvaluationVM/Tasks/BlendKeyframes.h"
 #include "MovieSceneAnimMixerSystem.generated.h"
 
 struct FAnimNextEvaluationTask;
@@ -117,6 +118,21 @@ struct FMovieSceneAccumulateAbsoluteBlendTask : public FAnimNextEvaluationTask
 	float ScaleFactor = 0.0f;
 };
 
+USTRUCT()
+struct FAnimNextBlendTwoKeyframesPreserveRootMotionTask : public FAnimNextBlendTwoKeyframesTask
+{
+	GENERATED_BODY()
+
+	DECLARE_ANIM_EVALUATION_TASK(FAnimNextBlendTwoKeyframesPreserveRootMotionTask)
+
+	static FAnimNextBlendTwoKeyframesPreserveRootMotionTask Make(float InterpolationAlpha);
+
+	// Task entry point
+	virtual void Execute(UE::AnimNext::FEvaluationVM& VM) const override;
+};
+
+
+
 // Takes in evaluation tasks with optional pose weight, masks, priority and a given animation target.
 // Constructs a hierarchical 'mixer' per bound object per target. 
 // Similar to blender systems, in a 'many to one' operation, each mixer will create an entity with a single evaluation task
