Sequencer: World Space root motions now compensate for component transform offsets and rotations

diff --git a/Engine/Plugins/Experimental/MovieSceneAnimMixer/Source/MovieSceneAnimMixer/Private/Systems/MovieSceneAnimMixerSystem.cpp b/Engine/Plugins/Experimental/MovieSceneAnimMixer/Source/MovieSceneAnimMixer/Private/Systems/MovieSceneAnimMixerSystem.cpp
--- a/Engine/Plugins/Experimental/MovieSceneAnimMixer/Source/MovieSceneAnimMixer/Private/Systems/MovieSceneAnimMixerSystem.cpp
+++ b/Engine/Plugins/Experimental/MovieSceneAnimMixer/Source/MovieSceneAnimMixer/Private/Systems/MovieSceneAnimMixerSystem.cpp
@@ -270,6 +270,10 @@ namespace UE::MovieScene
 								Conversion |= ESpaceConversions::TransformOriginToWorld;
 							}
 						}
+						else
+						{
+							Conversion |= ESpaceConversions::WorldSpaceComponentTransformCompensation;
+						}
 
 						if (Conversion != ESpaceConversions::None)
 						{
diff --git a/Engine/Plugins/Experimental/MovieSceneAnimMixer/Source/MovieSceneAnimMixer/Private/Systems/MovieSceneRootMotionSystem.cpp b/Engine/Plugins/Experimental/MovieSceneAnimMixer/Source/MovieSceneAnimMixer/Private/Systems/MovieSceneRootMotionSystem.cpp
--- a/Engine/Plugins/Experimental/MovieSceneAnimMixer/Source/MovieSceneAnimMixer/Private/Systems/MovieSceneRootMotionSystem.cpp
+++ b/Engine/Plugins/Experimental/MovieSceneAnimMixer/Source/MovieSceneAnimMixer/Private/Systems/MovieSceneRootMotionSystem.cpp
@@ -257,6 +257,12 @@ void FAnimNextConvertRootMotionToWorldSpaceTask::Execute(UE::AnimNext::FEvaluati
 		RelativeTransform *= TransformOrigin;
 	}
 
+	if (EnumHasAnyFlags(Conversions, ESpaceConversions::WorldSpaceComponentTransformCompensation) && RootMotionData)
+	{
+		FTransform ComponentToActor = RootMotionData->ActorTransform * RootMotionData->ComponentToWorldTransform.Inverse();
+		RelativeTransform = ComponentToActor * RelativeTransform;
+	}
+
 	RootMotionAttribute->Value = RelativeTransform;
 }
 
diff --git a/Engine/Plugins/Experimental/MovieSceneAnimMixer/Source/MovieSceneAnimMixer/Public/Systems/MovieSceneRootMotionSystem.h b/Engine/Plugins/Experimental/MovieSceneAnimMixer/Source/MovieSceneAnimMixer/Public/Systems/MovieSceneRootMotionSystem.h
--- a/Engine/Plugins/Experimental/MovieSceneAnimMixer/Source/MovieSceneAnimMixer/Public/Systems/MovieSceneRootMotionSystem.h
+++ b/Engine/Plugins/Experimental/MovieSceneAnimMixer/Source/MovieSceneAnimMixer/Public/Systems/MovieSceneRootMotionSystem.h
@@ -109,10 +109,12 @@ struct FAnimNextConvertRootMotionToWorldSpaceTask : public FAnimNextEvaluationTa
 		TransformOriginToWorld = 1 << 1,
 		/** Convert the root motion from component -> actor space using the inverse component rotation only */
 		ComponentToActorRotation = 1 << 2,
+		/** Compensate for component rotation and translation offsets when applying root motion in world space */
+		WorldSpaceComponentTransformCompensation = 1 << 3,
 		/** Apply RootBaseTransform as an offset around RootOffsetOrigin where Root = Root * RootBaseTransform */
-		RootTransformOffset = 1 << 3,
+		RootTransformOffset = 1 << 4,
 		/** Completely override the root transform with RootTransform */
-		RootTransformOverride = 1 << 4,
+		RootTransformOverride = 1 << 5,
 	};
 
 	FAnimNextConvertRootMotionToWorldSpaceTask() = default;
