Sequencer: Fix duplicate LOCTEXT

#rnx

diff --git a/Engine/Plugins/Experimental/MovieSceneAnimMixer/Source/MovieSceneAnimMixerEditor/Private/MovieSceneAnimationMixerTrackEditor.cpp b/Engine/Plugins/Experimental/MovieSceneAnimMixer/Source/MovieSceneAnimMixerEditor/Private/MovieSceneAnimationMixerTrackEditor.cpp
--- a/Engine/Plugins/Experimental/MovieSceneAnimMixer/Source/MovieSceneAnimMixerEditor/Private/MovieSceneAnimationMixerTrackEditor.cpp
+++ b/Engine/Plugins/Experimental/MovieSceneAnimMixer/Source/MovieSceneAnimMixerEditor/Private/MovieSceneAnimationMixerTrackEditor.cpp
@@ -174,7 +174,7 @@ public:
 	}
 	void SetRootTransformOverride()
 	{
-		FScopedTransaction Transaction(LOCTEXT("SetRootTransformOffset", "Override Root Transform"));
+		FScopedTransaction Transaction(LOCTEXT("SetRootTransformOverride", "Override Root Transform"));
 
 		UMovieSceneSkeletalAnimationSection*         AnimSection = Cast<UMovieSceneSkeletalAnimationSection>(WeakSection.Get());
 		UMovieSceneAnimationBaseTransformDecoration* Decoration  = FindOrCreateRootDecoration();
@@ -187,7 +187,7 @@ public:
 	}
 	void RecenterRootTransform()
 	{
-		FScopedTransaction Transaction(LOCTEXT("SetRootTransformOffset", "Override Root Transform"));
+		FScopedTransaction Transaction(LOCTEXT("RecenterRootTransform", "Recenter Root Transform"));
 
 		TSharedPtr<ISequencer> SequencerPtr = Sequencer.Pin();
 
