Sequencer: Fix conflicting LOCTEXT

#rnx

diff --git a/Engine/Plugins/Experimental/MovieSceneAnimMixer/Source/MovieSceneAnimMixerEditor/Private/MovieSceneAnimationMixerTrackEditor.cpp b/Engine/Plugins/Experimental/MovieSceneAnimMixer/Source/MovieSceneAnimMixerEditor/Private/MovieSceneAnimationMixerTrackEditor.cpp
--- a/Engine/Plugins/Experimental/MovieSceneAnimMixer/Source/MovieSceneAnimMixerEditor/Private/MovieSceneAnimationMixerTrackEditor.cpp
+++ b/Engine/Plugins/Experimental/MovieSceneAnimMixer/Source/MovieSceneAnimMixerEditor/Private/MovieSceneAnimationMixerTrackEditor.cpp
@@ -187,7 +187,7 @@ public:
 	}
 	void RecenterRootTransform()
 	{
-		FScopedTransaction Transaction(LOCTEXT("RecenterRootTransform", "Recenter Root Transform"));
+		FScopedTransaction Transaction(LOCTEXT("RecenterRootTransform", "Re-center Root Transform"));
 
 		TSharedPtr<ISequencer> SequencerPtr = Sequencer.Pin();
 
