Sequencer: Fixed Anim Mixer Component -> Actor conversions erroneously including attachment parent transforms

diff --git a/Engine/Plugins/Experimental/MovieSceneAnimMixer/Source/MovieSceneAnimMixer/Private/Systems/MovieSceneRootMotionSystem.cpp b/Engine/Plugins/Experimental/MovieSceneAnimMixer/Source/MovieSceneAnimMixer/Private/Systems/MovieSceneRootMotionSystem.cpp
--- a/Engine/Plugins/Experimental/MovieSceneAnimMixer/Source/MovieSceneAnimMixer/Private/Systems/MovieSceneRootMotionSystem.cpp
+++ b/Engine/Plugins/Experimental/MovieSceneAnimMixer/Source/MovieSceneAnimMixer/Private/Systems/MovieSceneRootMotionSystem.cpp
@@ -138,8 +138,6 @@ void FMovieSceneMixerRootMotionComponentData::Initialize()
 		return;
 	}
 
-	ComponentToWorldTransform = BoundObject->GetComponentToWorld();
-
 	// If we're applying to the actor or a root custom attribute, we need to factor out the actor->compont rotation
 	const bool bNeedInverseMeshRotation = RootDestination == EMovieSceneRootMotionDestination::Actor
 			|| RootDestination == EMovieSceneRootMotionDestination::Attribute;
@@ -147,6 +145,8 @@ void FMovieSceneMixerRootMotionComponentData::Initialize()
 	USceneComponent* RootComponent = BoundObject->GetOwner()->GetRootComponent();
 	if (RootComponent)
 	{
+		ComponentToActorTransform = BoundObject->GetComponentTransform().GetRelativeTransform(RootComponent->GetComponentTransform()).Inverse();
+
 		if (bNeedInverseMeshRotation && RootComponent != BoundObject)
 		{
 			InverseMeshToActorRotation = RootComponent->GetComponentTransform().GetRelativeTransformReverse(BoundObject->GetComponentTransform()).GetRotation();
@@ -259,8 +259,7 @@ void FAnimNextConvertRootMotionToWorldSpaceTask::Execute(UE::AnimNext::FEvaluati
 
 	if (EnumHasAnyFlags(Conversions, ESpaceConversions::WorldSpaceComponentTransformCompensation) && RootMotionData)
 	{
-		FTransform ComponentToActor = RootMotionData->ActorTransform * RootMotionData->ComponentToWorldTransform.Inverse();
-		RelativeTransform = ComponentToActor * RelativeTransform;
+		RelativeTransform = RootMotionData->ComponentToActorTransform * RelativeTransform;
 	}
 
 	RootMotionAttribute->Value = RelativeTransform;
diff --git a/Engine/Plugins/Experimental/MovieSceneAnimMixer/Source/MovieSceneAnimMixer/Public/Systems/MovieSceneRootMotionSystem.h b/Engine/Plugins/Experimental/MovieSceneAnimMixer/Source/MovieSceneAnimMixer/Public/Systems/MovieSceneRootMotionSystem.h
--- a/Engine/Plugins/Experimental/MovieSceneAnimMixer/Source/MovieSceneAnimMixer/Public/Systems/MovieSceneRootMotionSystem.h
+++ b/Engine/Plugins/Experimental/MovieSceneAnimMixer/Source/MovieSceneAnimMixer/Public/Systems/MovieSceneRootMotionSystem.h
@@ -77,7 +77,7 @@ public:
 	EMovieSceneRootMotionDestination RootDestination;
 
 	FTransform ActorTransform;
-	FTransform ComponentToWorldTransform;
+	FTransform ComponentToActorTransform;
 
 private:
 
