Root motion fixes for stitch track
-Root motion from stitch tracks is now marked as "authoritative"
-This is a new attribute added to the animation, when blending the keyframes with FAnimNextBlendTwoKeyframesPreserveRootMotionTask we now check if either keyframe is authoritative, and preserve that root motion.
-This prevents stitch section root motion from blending with the root motion from the animation section, since the stitch track is already computing the poisitional blend through motion matcing.

diff --git a/Engine/Plugins/Experimental/MovieSceneAnimMixer/Source/MovieSceneAnimMixer/Private/AnimMixerComponentTypes.cpp b/Engine/Plugins/Experimental/MovieSceneAnimMixer/Source/MovieSceneAnimMixer/Private/AnimMixerComponentTypes.cpp
--- a/Engine/Plugins/Experimental/MovieSceneAnimMixer/Source/MovieSceneAnimMixer/Private/AnimMixerComponentTypes.cpp
+++ b/Engine/Plugins/Experimental/MovieSceneAnimMixer/Source/MovieSceneAnimMixer/Private/AnimMixerComponentTypes.cpp
@@ -12,6 +12,7 @@ namespace UE::MovieScene
 
 	const UE::Anim::FAttributeId FAnimMixerComponentTypes::RootTransformAttributeId(TEXT("RootTransform"), FCompactPoseBoneIndex(0));
 	const UE::Anim::FAttributeId FAnimMixerComponentTypes::RootTransformWeightAttributeId(TEXT("RootTransformWeight"), FCompactPoseBoneIndex(0));
+	const UE::Anim::FAttributeId FAnimMixerComponentTypes::RootTransformIsAuthoritativeAttributeId(TEXT("RootTransformIsAuthoritative"), FCompactPoseBoneIndex(0));
 
 	static TUniquePtr<FAnimMixerComponentTypes> GAnimMixerComponentTypes;
 	static bool GAnimMixerComponentTypesDestroyed = false;
diff --git a/Engine/Plugins/Experimental/MovieSceneAnimMixer/Source/MovieSceneAnimMixer/Private/Systems/MovieSceneAnimMixerSystem.cpp b/Engine/Plugins/Experimental/MovieSceneAnimMixer/Source/MovieSceneAnimMixer/Private/Systems/MovieSceneAnimMixerSystem.cpp
--- a/Engine/Plugins/Experimental/MovieSceneAnimMixer/Source/MovieSceneAnimMixer/Private/Systems/MovieSceneAnimMixerSystem.cpp
+++ b/Engine/Plugins/Experimental/MovieSceneAnimMixer/Source/MovieSceneAnimMixer/Private/Systems/MovieSceneAnimMixerSystem.cpp
@@ -382,6 +382,9 @@ void FAnimNextBlendTwoKeyframesPreserveRootMotionTask::Execute(UE::AnimNext::FEv
 	// From this final pose, FAnimNextStoreRootTransformTask is able to accurately read both Sequencer's desired world space
 	//    root transform, its desired weight, and the incoming desired root motion delta from locomotion. It can then blend
 	//    all these things together to form the final root motion delta which can be consumed by external systems like Mover
+	//
+	// If there are competing sources of Sequencer root motion, one may be authoritative, so blending will also be skipped
+	//    in this case, but only for the RootMotionTransform, since Sequencer doesn't directly write to RootMotionDelta
 
 	{
 		const TUniquePtr<FKeyframeState>* KeyframeA = VM.PeekValue<TUniquePtr<FKeyframeState>>(KEYFRAME_STACK_NAME, 0);
@@ -393,21 +396,33 @@ void FAnimNextBlendTwoKeyframesPreserveRootMotionTask::Execute(UE::AnimNext::FEv
 
 			const FTransformAnimationAttribute* RootMotionTransformA = AttributesA.Find<FTransformAnimationAttribute>(AnimMixerComponents->RootTransformAttributeId);
 			const FTransformAnimationAttribute* RootMotionDeltaA     = AttributesA.Find<FTransformAnimationAttribute>(RootMotionDeltaAttributeId);
+			const FIntegerAnimationAttribute* RootMotionIsAuthoritativeA = AttributesA.Find<FIntegerAnimationAttribute>(AnimMixerComponents->RootTransformIsAuthoritativeAttributeId);
 
 			const FTransformAnimationAttribute* RootMotionTransformB = AttributesB.Find<FTransformAnimationAttribute>(AnimMixerComponents->RootTransformAttributeId);
 			const FTransformAnimationAttribute* RootMotionDeltaB     = AttributesB.Find<FTransformAnimationAttribute>(RootMotionDeltaAttributeId);
+			const FIntegerAnimationAttribute* RootMotionIsAuthoritativeB = AttributesB.Find<FIntegerAnimationAttribute>(AnimMixerComponents->RootTransformIsAuthoritativeAttributeId);
 
+			// If Keyframe A is authoritative, or there is no transform from B, preserve only A. Otherwise it would be blended with identity, or non-authoritative source of root motion.
+			const bool bAIsAuthoritative = RootMotionTransformA && RootMotionIsAuthoritativeA && !RootMotionIsAuthoritativeB;
+			const bool bPreserveTransformA =  bAIsAuthoritative || (RootMotionTransformA && !RootMotionTransformB);
+
+			// Only preserve the transform from B if it is authoritative.
+			const bool bPreserveTransformB = RootMotionTransformB && RootMotionIsAuthoritativeB && !RootMotionIsAuthoritativeA;
+			
+			if (bPreserveTransformA)
+			{
+				OverrideRootTransform = RootMotionTransformA->Value;
+			}
+			else if (bPreserveTransformB)
+			{
+				OverrideRootTransform = RootMotionTransformB->Value;
+			}
+			
 			if (!RootMotionDeltaA && RootMotionDeltaB)
 			{
 				// Preserve root motion *delta* from B if it's not in A (ie, don't blend it with identity!)
 				OverrideRootDeltaTransform = RootMotionDeltaB->Value;
 			}
-
-			if (RootMotionTransformA && !RootMotionTransformB)
-			{
-				// Preserve root motion transform from A if it's not in B (ie, don't blend it with identity!)
-				OverrideRootTransform = RootMotionTransformA->Value;
-			}
 		}
 	}
 
diff --git a/Engine/Plugins/Experimental/MovieSceneAnimMixer/Source/MovieSceneAnimMixer/Public/AnimMixerComponentTypes.h b/Engine/Plugins/Experimental/MovieSceneAnimMixer/Source/MovieSceneAnimMixer/Public/AnimMixerComponentTypes.h
--- a/Engine/Plugins/Experimental/MovieSceneAnimMixer/Source/MovieSceneAnimMixer/Public/AnimMixerComponentTypes.h
+++ b/Engine/Plugins/Experimental/MovieSceneAnimMixer/Source/MovieSceneAnimMixer/Public/AnimMixerComponentTypes.h
@@ -27,6 +27,9 @@ namespace UE::MovieScene
 		// A root motion transform in animation space to be applied (after space conversion)
 		MOVIESCENEANIMMIXER_API static const UE::Anim::FAttributeId RootTransformAttributeId;
 		MOVIESCENEANIMMIXER_API static const UE::Anim::FAttributeId RootTransformWeightAttributeId;
+		// Internal flag marking a section as authoritative source of root motion.
+		// Some sections, e.g. stitch sections, should not have their root motion blended with others, since it's using motion matching to blend into the animation already.
+		MOVIESCENEANIMMIXER_API static const UE::Anim::FAttributeId RootTransformIsAuthoritativeAttributeId;
 
 		TComponentTypeID<int32> Priority;
 		TComponentTypeID<TInstancedStruct<FMovieSceneMixedAnimationTarget>> Target;
diff --git a/Engine/Plugins/Experimental/MovieScenePoseSearchTracks/Source/MovieScenePoseSearchTracks/Private/Sections/MovieSceneStitchAnimSection.cpp b/Engine/Plugins/Experimental/MovieScenePoseSearchTracks/Source/MovieScenePoseSearchTracks/Private/Sections/MovieSceneStitchAnimSection.cpp
--- a/Engine/Plugins/Experimental/MovieScenePoseSearchTracks/Source/MovieScenePoseSearchTracks/Private/Sections/MovieSceneStitchAnimSection.cpp
+++ b/Engine/Plugins/Experimental/MovieScenePoseSearchTracks/Source/MovieScenePoseSearchTracks/Private/Sections/MovieSceneStitchAnimSection.cpp
@@ -237,6 +237,9 @@ void FMovieSceneStitchAnimEvaluationTask::Execute(UE::AnimNext::FEvaluationVM& V
 
 		FFloatAnimationAttribute* RootMotionWeightAttribute = Keyframe.Attributes.FindOrAdd<FFloatAnimationAttribute>(AnimMixerComponents->RootTransformWeightAttributeId);
 		RootMotionWeightAttribute->Value = 1.f;
+
+		FIntegerAnimationAttribute* RootMotionIsAuthoritativeAttribute = Keyframe.Attributes.FindOrAdd<FIntegerAnimationAttribute>(AnimMixerComponents->RootTransformIsAuthoritativeAttributeId);
+		RootMotionIsAuthoritativeAttribute->Value = 1;
 	}
 
 	VM.PushValue(KEYFRAME_STACK_NAME, MakeUnique<FKeyframeState>(MoveTemp(Keyframe)));
