Sequencer: Preserve authoritative root motion attribute when blending poses that themselves have authoritative root motion

diff --git a/Engine/Plugins/Experimental/MovieSceneAnimMixer/Source/MovieSceneAnimMixer/Private/Systems/MovieSceneAnimMixerSystem.cpp b/Engine/Plugins/Experimental/MovieSceneAnimMixer/Source/MovieSceneAnimMixer/Private/Systems/MovieSceneAnimMixerSystem.cpp
--- a/Engine/Plugins/Experimental/MovieSceneAnimMixer/Source/MovieSceneAnimMixer/Private/Systems/MovieSceneAnimMixerSystem.cpp
+++ b/Engine/Plugins/Experimental/MovieSceneAnimMixer/Source/MovieSceneAnimMixer/Private/Systems/MovieSceneAnimMixerSystem.cpp
@@ -351,6 +351,7 @@ void FAnimNextBlendTwoKeyframesPreserveRootMotionTask::Execute(UE::AnimNext::FEv
 	TOptional<FTransform> OverrideRootTransform;
 	TOptional<float>      OverrideRootWeight;
 	TOptional<FTransform> OverrideRootDeltaTransform;
+	bool                  bIsAuthoritative = false;
 
 	// This task is used to perform a two way blend between 2 poses without blending root motion attributes
 	//    if they do not exist in either pose. By default these attributes would blend with the identity matrix
@@ -423,6 +424,8 @@ void FAnimNextBlendTwoKeyframesPreserveRootMotionTask::Execute(UE::AnimNext::FEv
 				// Preserve root motion *delta* from B if it's not in A (ie, don't blend it with identity!)
 				OverrideRootDeltaTransform = RootMotionDeltaB->Value;
 			}
+
+			bIsAuthoritative = RootMotionIsAuthoritativeA || RootMotionIsAuthoritativeB;
 		}
 	}
 
@@ -444,6 +447,11 @@ void FAnimNextBlendTwoKeyframesPreserveRootMotionTask::Execute(UE::AnimNext::FEv
 				Keyframe->Attributes.FindOrAdd<FTransformAnimationAttribute>(RootMotionDeltaAttributeId)->Value
 					= OverrideRootDeltaTransform.GetValue();
 			}
+			if (bIsAuthoritative)
+			{
+				Keyframe->Attributes.FindOrAdd<FIntegerAnimationAttribute>(AnimMixerComponents->RootTransformIsAuthoritativeAttributeId)->Value
+					= 1;
+			}
 
 			VM.PushValue(KEYFRAME_STACK_NAME, MoveTemp(Keyframe));
 		}
