Add/fix profiler scopes for two game thread animnext update functions

diff --git a/Engine/Plugins/Experimental/MovieSceneAnimMixer/Source/MovieSceneAnimMixer/Private/Systems/MovieSceneRootMotionSystem.cpp b/Engine/Plugins/Experimental/MovieSceneAnimMixer/Source/MovieSceneAnimMixer/Private/Systems/MovieSceneRootMotionSystem.cpp
--- a/Engine/Plugins/Experimental/MovieSceneAnimMixer/Source/MovieSceneAnimMixer/Private/Systems/MovieSceneRootMotionSystem.cpp
+++ b/Engine/Plugins/Experimental/MovieSceneAnimMixer/Source/MovieSceneAnimMixer/Private/Systems/MovieSceneRootMotionSystem.cpp
@@ -411,6 +411,7 @@ void FAnimNextStoreRootTransformTask::Execute(UE::AnimNext::FEvaluationVM& VM) c
 
 		auto ApplyRootMotion = [RootTransform, WeakTarget = RootMotionData->Target]
 		{
+			SCOPED_NAMED_EVENT(AnimNextStoreRootTransformTask_ApplyRootMotion, FColor::Orange);
 			USceneComponent* TargetComponent = WeakTarget.Get();
 			if (TargetComponent && !GActorMovementTracker.IsActive())
 			{
