Sequencer: Prevent double transform updates when animating root motion along with a transform track

diff --git a/Engine/Plugins/Experimental/MovieSceneAnimMixer/Source/MovieSceneAnimMixer/Private/Systems/MovieSceneAnimMixerSystem.cpp b/Engine/Plugins/Experimental/MovieSceneAnimMixer/Source/MovieSceneAnimMixer/Private/Systems/MovieSceneAnimMixerSystem.cpp
--- a/Engine/Plugins/Experimental/MovieSceneAnimMixer/Source/MovieSceneAnimMixer/Private/Systems/MovieSceneAnimMixerSystem.cpp
+++ b/Engine/Plugins/Experimental/MovieSceneAnimMixer/Source/MovieSceneAnimMixer/Private/Systems/MovieSceneAnimMixerSystem.cpp
@@ -70,11 +70,13 @@ namespace UE::MovieScene
 	struct FEvaluateAnimMixers
 	{
 		TMap<FMovieSceneAnimMixerKey, TSharedPtr<FMovieSceneAnimMixer>>* Mixers;
+		TMap<FObjectKey, TWeakPtr<FMovieSceneMixerRootMotionComponentData>>* RootMotions;
 		const UMovieSceneRootMotionSystem* RootMotionSystem;
 		const UMovieSceneTransformOriginSystem* TransformOriginSystem;
 
-		FEvaluateAnimMixers(TMap<FMovieSceneAnimMixerKey, TSharedPtr<FMovieSceneAnimMixer>>* InMixers, UMovieSceneEntitySystemLinker* InLinker)
+		FEvaluateAnimMixers(TMap<FMovieSceneAnimMixerKey, TSharedPtr<FMovieSceneAnimMixer>>* InMixers, TMap<FObjectKey, TWeakPtr<FMovieSceneMixerRootMotionComponentData>>* InRootMotions, UMovieSceneEntitySystemLinker* InLinker)
 			: Mixers(InMixers)
+			, RootMotions(InRootMotions)
 		{
 			RootMotionSystem = InLinker->FindSystem<UMovieSceneRootMotionSystem>();
 			TransformOriginSystem = InLinker->FindSystem<UMovieSceneTransformOriginSystem>();
@@ -177,7 +179,7 @@ namespace UE::MovieScene
 				return false;
 			};
 
-			TSharedPtr<FMovieSceneMixerRootMotionComponentData> RootMotion = Mixer->WeakRootMotion.Pin();
+			TSharedPtr<FMovieSceneMixerRootMotionComponentData> RootMotion = RootMotions->FindRef(MeshComponent.GetObject()).Pin();
 
 			for (int32 Index = 0; Index < NumMixerEntries; )
 			{
@@ -537,34 +539,58 @@ UMovieSceneAnimMixerSystem::UMovieSceneAnimMixerSystem(const FObjectInitializer&
 	}
 }
 
-TSharedPtr<FMovieSceneMixerRootMotionComponentData> UMovieSceneAnimMixerSystem::FindRootMotion(FMovieSceneAnimMixerKey InMixerKey) const
+TSharedPtr<FMovieSceneMixerRootMotionComponentData> UMovieSceneAnimMixerSystem::FindRootMotion(FObjectKey InObject) const
 {
-	TSharedPtr<FMovieSceneAnimMixer> Mixer = Mixers.FindRef(InMixerKey);
+	TSharedPtr<FMovieSceneMixerRootMotionComponentData> RootMotion = RootMotionByObject.FindRef(InObject).Pin();
+	if (RootMotion)
+	{
+		return RootMotion;
+	}
+	
+	if (USceneComponent* SceneComponent = Cast<USceneComponent>(InObject.ResolveObjectPtr()))
+	{
+		AActor* Actor = SceneComponent ? SceneComponent->GetOwner() : nullptr;
+		if (!Actor || SceneComponent != Actor->GetRootComponent())
+		{
+			return nullptr;
+		}
 
-	return Mixer ? Mixer->WeakRootMotion.Pin() : TSharedPtr<FMovieSceneMixerRootMotionComponentData>();
+		FObjectKey ActorKey(Actor);
+		for (auto It = ActorToRootMotion.CreateConstKeyIterator(ActorKey); It; ++It)
+		{
+			TSharedPtr<FMovieSceneMixerRootMotionComponentData> ThisRootMotion = It.Value().Pin();
+			if (ThisRootMotion && ThisRootMotion->RootDestination == EMovieSceneRootMotionDestination::Actor)
+			{
+				return ThisRootMotion;
+			}
+		}
+	}
+	return nullptr;
 }
 
-void UMovieSceneAnimMixerSystem::AssignRootMotion(FMovieSceneAnimMixerKey InMixerKey, TSharedPtr<FMovieSceneMixerRootMotionComponentData> RootMotion)
+void UMovieSceneAnimMixerSystem::AssignRootMotion(FObjectKey InObject, TSharedPtr<FMovieSceneMixerRootMotionComponentData> RootMotion)
 {
-	TSharedPtr<FMovieSceneAnimMixer> Mixer = Mixers.FindRef(InMixerKey);
-
-	if (!Mixer && RootMotion)
+	if (RootMotion)
 	{
-		Mixer = MakeShared<FMovieSceneAnimMixer>();
-		Mixers.Add(InMixerKey, Mixer);
+		RootMotionByObject.Add(InObject, RootMotion);
+		USceneComponent* SceneComponent = Cast<USceneComponent>(InObject.ResolveObjectPtr());
+		AActor* Actor = SceneComponent ? SceneComponent->GetOwner() : nullptr;
+		if (Actor && !ActorToRootMotion.FindPair(Actor, RootMotion))
+		{
+			ActorToRootMotion.Add(Actor, RootMotion);
+		}
 	}
-	if (Mixer)
+	else
 	{
-		Mixer->bNeedsResort |= (Mixer->WeakRootMotion.Pin() != RootMotion);
-		Mixer->WeakRootMotion = RootMotion;
+		RootMotionByObject.Remove(InObject);
 	}
 }
 
 void UMovieSceneAnimMixerSystem::InitializeAllRootMotion()
 {
-	for (auto It = Mixers.CreateIterator(); It; ++It)
+	for (auto It = RootMotionByObject.CreateIterator(); It; ++It)
 	{
-		TSharedPtr<FMovieSceneMixerRootMotionComponentData> RootMotion = It.Value()->WeakRootMotion.Pin();;
+		TSharedPtr<FMovieSceneMixerRootMotionComponentData> RootMotion = It.Value().Pin();
 		if (RootMotion)
 		{
 			RootMotion->Initialize();
@@ -675,15 +701,14 @@ void UMovieSceneAnimMixerSystem::OnRun(FSystemTaskPrerequisites& InPrerequisites
 
 					// If we know we'll need root motion, ensure it is set up correctly with a lifetime reference that keeps it
 					//     alive as long as this entry
-					TSharedPtr<FMovieSceneMixerRootMotionComponentData> RootMotion = Mixer->WeakRootMotion.Pin();
+					TSharedPtr<FMovieSceneMixerRootMotionComponentData> RootMotion = FindRootMotion(BoundObjectKey);
 					if (!RootMotion)
 					{
 						RootMotion = MakeShared<FMovieSceneMixerRootMotionComponentData>();
 
 						RootMotion->RootDestination = EMovieSceneRootMotionDestination::RootBone;
 						RootMotion->OriginalBoundObject = Cast<USceneComponent>(BoundObjectKey.ResolveObjectPtr());
-
-						Mixer->WeakRootMotion = RootMotion;
+						AssignRootMotion(BoundObjectKey, RootMotion);
 					}
 
 					// Overwrite the RootDestination for the root motion if we have a legacy setting.
@@ -783,13 +808,91 @@ void UMovieSceneAnimMixerSystem::OnRun(FSystemTaskPrerequisites& InPrerequisites
 					Mixer->MixerEntityID = FMovieSceneEntityID();
 				}
 
-				// If we don't have any root motion then we can clean this up as well
-				if (Mixer->WeakRootMotion.Pin() == nullptr)
+				It.RemoveCurrent();
+			}
+		}
+
+		for (auto It = RootMotionByObject.CreateIterator(); It; ++It)
+		{
+			if (It.Value().Pin() == nullptr)
+			{
+				It.RemoveCurrent();
+			}
+		}
+		for (auto It = ActorToRootMotion.CreateIterator(); It; ++It)
+		{
+			if (It.Value().Pin() == nullptr)
+			{
+				It.RemoveCurrent();
+			}
+		}
+
+		// ----------------------------------------------------------------------------------------------------------------------------------------
+		// Step 5. Mutate transforms as necessary to include their root motion
+		struct FMutateTransforms : IMovieSceneConditionalEntityMutation
+		{
+			const TMap<FObjectKey, TWeakPtr<FMovieSceneMixerRootMotionComponentData>>* RootMotionByObject;
+			const TMultiMap<FObjectKey, TWeakPtr<FMovieSceneMixerRootMotionComponentData>>* ActorToRootMotion;
+
+			FComponentTypeID AnimMixerPoseProducer;
+
+			FMutateTransforms(
+				const TMap<FObjectKey, TWeakPtr<FMovieSceneMixerRootMotionComponentData>>* InRootMotionByObject,
+				const TMultiMap<FObjectKey, TWeakPtr<FMovieSceneMixerRootMotionComponentData>>* InActorToRootMotion)
+				: RootMotionByObject(InRootMotionByObject)
+				, ActorToRootMotion(InActorToRootMotion)
+				, AnimMixerPoseProducer(FMovieSceneTracksComponentTypes::Get()->Tags.AnimMixerPoseProducer)
+			{
+			}
+
+			virtual void CreateMutation(FEntityManager* EntityManager, FComponentMask* InOutEntityComponentTypes) const override
+			{
+				if (InOutEntityComponentTypes->Contains(AnimMixerPoseProducer))
+				{
+					InOutEntityComponentTypes->Remove(AnimMixerPoseProducer);
+				}
+				else
 				{
-					It.RemoveCurrent();
+					InOutEntityComponentTypes->Set(AnimMixerPoseProducer);
 				}
 			}
-		}
+			virtual void MarkAllocation(FEntityAllocation* Allocation, TBitArray<>& OutEntitiesToMutate) const
+			{
+				TComponentReader<UObject*> BoundObjects = Allocation->ReadComponents(FBuiltInComponentTypes::Get()->BoundObject);
+				const bool bAlreadyRedirectingToRoot = Allocation->HasComponent(AnimMixerPoseProducer);
+
+				const int32 NumEntities = Allocation->Num();
+				for (int32 Index = 0; Index < NumEntities; ++Index)
+				{
+					USceneComponent* SceneComponent = Cast<USceneComponent>(BoundObjects[Index]);
+
+					if (!SceneComponent)
+					{
+						continue;
+					}
+
+					bool bNeedsTracking = RootMotionByObject->Contains(SceneComponent);
+					if (!bNeedsTracking)
+					{
+						if (AActor* Owner = SceneComponent->GetOwner())
+						{
+							bNeedsTracking = Owner->GetRootComponent() == SceneComponent && ActorToRootMotion->Contains(Owner);
+						}
+					}
+
+					if (bAlreadyRedirectingToRoot != bNeedsTracking)
+					{
+						OutEntitiesToMutate.PadToNum(Index+1, false);
+						OutEntitiesToMutate[Index] = true;
+					}
+				}
+			}
+		};
+
+		FMutateTransforms Mutation(&RootMotionByObject, &ActorToRootMotion);
+		FEntityComponentFilter Filter;
+		Filter.All({ BuiltInComponents->BoundObject, TrackComponents->ComponentTransform.PropertyTag, BuiltInComponents->CustomPropertyIndex });
+		Linker->EntityManager.MutateConditional(Filter, Mutation);
 	}
 }
 
@@ -820,7 +923,7 @@ void UMovieSceneAnimMixerSystem::OnSchedulePersistentTasks(UE::MovieScene::IEnti
 		.Read(AnimMixerComponents->MeshComponent)
 		.Read(AnimMixerComponents->Target)
 		.Write(AnimMixerComponents->MixerTask)
-		.Schedule_PerEntity<FEvaluateAnimMixers>(&Linker->EntityManager, TaskScheduler, &Mixers, Linker);
+		.Schedule_PerEntity<FEvaluateAnimMixers>(&Linker->EntityManager, TaskScheduler, &Mixers, &RootMotionByObject, Linker);
 
 	TaskScheduler->AddPrerequisite(UpdateTask, MixTask);
 }
@@ -858,6 +961,7 @@ void UMovieSceneAnimMixerSystem::OnCleanTaggedGarbage()
 	.Iterate_PerEntity(&Linker->EntityManager,
 		[this](FObjectComponent MeshComponent, const TInstancedStruct<FMovieSceneMixedAnimationTarget>& Target)
 		{
+			RootMotionByObject.Remove(MeshComponent.GetObject());
 			Mixers.Remove(FMovieSceneAnimMixerKey({ MeshComponent.GetObject(), Target }));
 		}
 	);
diff --git a/Engine/Plugins/Experimental/MovieSceneAnimMixer/Source/MovieSceneAnimMixer/Private/Systems/MovieSceneRootMotionSystem.cpp b/Engine/Plugins/Experimental/MovieSceneAnimMixer/Source/MovieSceneAnimMixer/Private/Systems/MovieSceneRootMotionSystem.cpp
--- a/Engine/Plugins/Experimental/MovieSceneAnimMixer/Source/MovieSceneAnimMixer/Private/Systems/MovieSceneRootMotionSystem.cpp
+++ b/Engine/Plugins/Experimental/MovieSceneAnimMixer/Source/MovieSceneAnimMixer/Private/Systems/MovieSceneRootMotionSystem.cpp
@@ -62,12 +62,73 @@ namespace UE::MovieScene
 	struct FInitializeRootMotionTask
 	{
 		UMovieSceneAnimMixerSystem* AnimMixer;
-		
-		void Run(FEntityAllocationWriteContext WriteContext) const
+
+		FInitializeRootMotionTask(UMovieSceneAnimMixerSystem* InAnimMixer)
+			: AnimMixer(InAnimMixer)
+		{}
+
+		void PreTask()
 		{
 			AnimMixer->InitializeAllRootMotion();
 		}
 
+		void ForEachEntity(UObject* BoundObject,
+			const double* InLocationX, const double* InLocationY, const double* InLocationZ,
+			const double* InRotationX, const double* InRotationY, const double* InRotationZ,
+			const double* InScaleX, const double* InScaleY, const double* InScaleZ) const
+		{
+			USceneComponent* SceneComponent = Cast<USceneComponent>(BoundObject);
+			if (!SceneComponent)
+			{
+				return;
+			}
+
+			USceneComponent* RootComponent  = SceneComponent ? SceneComponent->GetOwner()->GetRootComponent() : SceneComponent;
+			if (BoundObject == RootComponent)
+			{
+				FObjectKey BoundObjectKey(BoundObject);
+
+				TSharedPtr<FMovieSceneMixerRootMotionComponentData> RootMotion = AnimMixer->FindRootMotion(BoundObjectKey);
+				if (RootMotion && RootMotion->RootDestination == EMovieSceneRootMotionDestination::Actor)
+				{
+					FVector  Location = RootMotion->ActorTransform.GetLocation();
+					FRotator Rotation = RootMotion->ActorTransform.GetRotation().Rotator();
+					FVector  Scale    = RootMotion->ActorTransform.GetScale3D();
+
+					if (InLocationX) { Location.X = *InLocationX; }
+					if (InLocationY) { Location.Y = *InLocationY; }
+					if (InLocationZ) { Location.Z = *InLocationZ; }
+
+					if (InRotationX) { Rotation.Roll  = *InRotationX; }
+					if (InRotationY) { Rotation.Pitch = *InRotationY; }
+					if (InRotationZ) { Rotation.Yaw   = *InRotationZ; }
+
+					if (InScaleX) { Scale.X = *InScaleX; }
+					if (InScaleY) { Scale.Y = *InScaleY; }
+					if (InScaleZ) { Scale.Z = *InScaleZ; }
+
+					RootMotion->ActorTransform = FTransform(Rotation, Location, Scale);
+					return;
+				}
+			}
+
+			FTransform Transform = SceneComponent->GetRelativeTransform();
+			FIntermediate3DTransform IntermediateTransform(Transform.GetLocation(), Transform.GetRotation().Rotator(), Transform.GetScale3D());
+
+			if (InLocationX) { IntermediateTransform.T_X = *InLocationX; }
+			if (InLocationY) { IntermediateTransform.T_Y = *InLocationY; }
+			if (InLocationZ) { IntermediateTransform.T_Z = *InLocationZ; }
+
+			if (InRotationX) { IntermediateTransform.R_X = *InRotationX; }
+			if (InRotationY) { IntermediateTransform.R_Y = *InRotationY; }
+			if (InRotationZ) { IntermediateTransform.R_Z = *InRotationZ; }
+
+			if (InScaleX) { IntermediateTransform.S_X = *InScaleX; }
+			if (InScaleY) { IntermediateTransform.S_Y = *InScaleY; }
+			if (InScaleZ) { IntermediateTransform.S_Z = *InScaleZ; }
+
+			UE::MovieScene::FIntermediate3DTransform::ApplyTransformTo(SceneComponent, IntermediateTransform);
+		}
 	};
 
 	struct FIsActorBeingMovedStatics
@@ -386,17 +447,17 @@ void FAnimNextStoreRootTransformTask::Execute(UE::AnimNext::FEvaluationVM& VM) c
 				static const TCHAR* LogName = TEXT("MovieSceneRootMotion");
 				USceneComponent* Component = RootMotionData->OriginalBoundObject.Get();
 
-				auto DrawMarker = [Component](const FTransform& Transform, const FColor& Color, auto&& Name)
+				auto DrawMarker = [Component](const FTransform& Transform, const FColor& Color, auto&&... Args)
 				{
 					FVector Dir(50.f, 0.f, 0.f);
 
-					UE_VLOG_CIRCLE(Component, LogName, Display, Transform.GetLocation(), FVector::UpVector, 10.f, Color, Name);
+					UE_VLOG_CIRCLE(Component, LogName, Display, Transform.GetLocation(), FVector::UpVector, 10.f, Color, Args...);
 					UE_VLOG_ARROW(Component, LogName, Display, Transform.GetLocation(), Transform.GetLocation() + Transform.GetRotation()*Dir, Color, TEXT(""));
 				};
 
 				DrawMarker(ActorHeadingTransform,                   FColorList::Black, TEXT("Actor"));
-				DrawMarker(RootTransform,                           FColorList::Blue,  TEXT("Sequencer"));
-				DrawMarker(LocomotionDelta * ActorHeadingTransform, FColorList::Red,   TEXT("Locomotion"));
+				DrawMarker(RootTransform,                           FColorList::Blue,  TEXT("Sequencer (%.2f%%"),      TransformWeight *100.f);
+				DrawMarker(LocomotionDelta * ActorHeadingTransform, FColorList::Red,   TEXT("Locomotion %.2f%%"), (1.f-TransformWeight)*100.f);
 				DrawMarker(BlendResult     * ActorHeadingTransform, FColorList::Green, TEXT("Result"));
 			}
 #endif
@@ -471,7 +532,7 @@ void UMovieSceneRootMotionSystem::OnRun(FSystemTaskPrerequisites& InPrerequisite
 	if (Runner->GetCurrentPhase() == ESystemPhase::Instantiation)
 	{
 		UMovieSceneAnimMixerSystem* AnimMixer = Linker->FindSystem<UMovieSceneAnimMixerSystem>();
-		if (!ensure(AnimMixer))
+		if (!AnimMixer)
 		{
 			// Should never exist without the anim mixer
 			return;
@@ -563,9 +624,9 @@ void UMovieSceneRootMotionSystem::OnRun(FSystemTaskPrerequisites& InPrerequisite
 
 					if (OutRootMotion)
 					{
-						AnimMixer->AssignRootMotion(Key, OutRootMotion);
+						AnimMixer->AssignRootMotion(BoundObjectKey, OutRootMotion);
 					}
-					else if (TSharedPtr<FMovieSceneMixerRootMotionComponentData> RootMotion = AnimMixer->FindRootMotion(Key))
+					else if (TSharedPtr<FMovieSceneMixerRootMotionComponentData> RootMotion = AnimMixer->FindRootMotion(BoundObjectKey))
 					{
 						OutRootMotion = RootMotion;
 					}
@@ -573,7 +634,7 @@ void UMovieSceneRootMotionSystem::OnRun(FSystemTaskPrerequisites& InPrerequisite
 					{
 						// Make a new one
 						OutRootMotion = MakeShared<FMovieSceneMixerRootMotionComponentData>();
-						AnimMixer->AssignRootMotion(Key, OutRootMotion);
+						AnimMixer->AssignRootMotion(BoundObjectKey, OutRootMotion);
 					}
 
 					USceneComponent* Target = Cast<USceneComponent>(BoundObjectKey.ResolveObjectPtr());
@@ -613,7 +674,7 @@ void UMovieSceneRootMotionSystem::OnSchedulePersistentTasks(UE::MovieScene::IEnt
 	FAnimMixerComponentTypes* AnimMixerComponents = FAnimMixerComponentTypes::Get();
 
 	UMovieSceneAnimMixerSystem* AnimMixer = Linker->FindSystem<UMovieSceneAnimMixerSystem>();
-	if (!ensure(AnimMixer))
+	if (!AnimMixer)
 	{
 		return;
 	}
@@ -635,7 +696,19 @@ void UMovieSceneRootMotionSystem::OnSchedulePersistentTasks(UE::MovieScene::IEnt
 	FTaskParams InitialzeParams(TEXT("Initialize Root Motion"));
 	InitialzeParams.ForceGameThread();
 	InitialzeParams.bForcePropagateDownstream = true;
-	FTaskID InitializeRootMotion = TaskScheduler->AddTask<FInitializeRootMotionTask>(InitialzeParams, AnimMixer);
+	InitialzeParams.bForcePrePostTask = true;
+
+	FTaskID InitializeRootMotion = FEntityTaskBuilder()
+	.Read(BuiltInComponents->BoundObject)
+	.ReadAnyOf(
+		BuiltInComponents->DoubleResult[0], BuiltInComponents->DoubleResult[1], BuiltInComponents->DoubleResult[2],
+		BuiltInComponents->DoubleResult[3], BuiltInComponents->DoubleResult[4], BuiltInComponents->DoubleResult[5],
+		BuiltInComponents->DoubleResult[6], BuiltInComponents->DoubleResult[7], BuiltInComponents->DoubleResult[8]
+	)
+	.FilterAll({ TrackComponents->ComponentTransform.PropertyTag, BuiltInComponents->CustomPropertyIndex, TrackComponents->Tags.AnimMixerPoseProducer })
+	.FilterNone({ BuiltInComponents->Tags.Ignored })
+	.SetParams(InitialzeParams)
+	.Schedule_PerEntity<FInitializeRootMotionTask>(&Linker->EntityManager, TaskScheduler, AnimMixer);
 
 	TaskScheduler->AddPrerequisite(WaitForAllTransforms,      InitializeRootMotion);
 	TaskScheduler->AddPrerequisite(GatherRootDestinationTask, InitializeRootMotion);
diff --git a/Engine/Plugins/Experimental/MovieSceneAnimMixer/Source/MovieSceneAnimMixer/Public/Systems/MovieSceneAnimMixerSystem.h b/Engine/Plugins/Experimental/MovieSceneAnimMixer/Source/MovieSceneAnimMixer/Public/Systems/MovieSceneAnimMixerSystem.h
--- a/Engine/Plugins/Experimental/MovieSceneAnimMixer/Source/MovieSceneAnimMixer/Public/Systems/MovieSceneAnimMixerSystem.h
+++ b/Engine/Plugins/Experimental/MovieSceneAnimMixer/Source/MovieSceneAnimMixer/Public/Systems/MovieSceneAnimMixerSystem.h
@@ -150,8 +150,8 @@ public:
 
 	static TInstancedStruct<FMovieSceneMixedAnimationTarget> ResolveAnimationTarget(FObjectKey ObjectKey, const TInstancedStruct<FMovieSceneMixedAnimationTarget>& InTarget);
 
-	TSharedPtr<FMovieSceneMixerRootMotionComponentData> FindRootMotion(FMovieSceneAnimMixerKey InMixerKey) const;
-	void AssignRootMotion(FMovieSceneAnimMixerKey InMixerKey, TSharedPtr<FMovieSceneMixerRootMotionComponentData> RootMotion);
+	TSharedPtr<FMovieSceneMixerRootMotionComponentData> FindRootMotion(FObjectKey InObject) const;
+	void AssignRootMotion(FObjectKey InObjectKey, TSharedPtr<FMovieSceneMixerRootMotionComponentData> RootMotion);
 
 	void InitializeAllRootMotion();
 
@@ -160,6 +160,12 @@ private:
 	// Map of animation mixers
 	TMap<FMovieSceneAnimMixerKey, TSharedPtr<FMovieSceneAnimMixer>> Mixers;
 
+	// Map of Root Motion data for each object
+	TMap<FObjectKey, TWeakPtr<FMovieSceneMixerRootMotionComponentData>> RootMotionByObject;
+
+	// Map from actor to all the root motions that may contribute to the actor's transform
+	TMultiMap<FObjectKey, TWeakPtr<FMovieSceneMixerRootMotionComponentData>> ActorToRootMotion;
+
 	virtual void OnLink() override;
 	virtual bool IsRelevantImpl(UMovieSceneEntitySystemLinker* InLinker) const override;
 	virtual void OnRun(FSystemTaskPrerequisites& InPrerequisites, FSystemSubsequentTasks& Subsequents) override;
diff --git a/Engine/Source/Runtime/MovieScene/Public/EntitySystem/MovieScenePropertyComponentHandler.h b/Engine/Source/Runtime/MovieScene/Public/EntitySystem/MovieScenePropertyComponentHandler.h
--- a/Engine/Source/Runtime/MovieScene/Public/EntitySystem/MovieScenePropertyComponentHandler.h
+++ b/Engine/Source/Runtime/MovieScene/Public/EntitySystem/MovieScenePropertyComponentHandler.h
@@ -367,6 +367,11 @@ struct TPropertyComponentHandlerImpl<PropertyTraits, TPropertyMetaData<MetaDataT
 	}
 
 	virtual void ScheduleSetterTasks(const FPropertyDefinition& Definition, TArrayView<const FPropertyCompositeDefinition> Composites, const FPropertyStats& Stats, IEntitySystemScheduler* TaskScheduler, UMovieSceneEntitySystemLinker* Linker)
+	{
+		ScheduleSetterTasksImpl(Definition, Composites, Stats, TaskScheduler, Linker, FEntityComponentFilter());
+	}
+
+	void ScheduleSetterTasksImpl(const FPropertyDefinition& Definition, TArrayView<const FPropertyCompositeDefinition> Composites, const FPropertyStats& Stats, IEntitySystemScheduler* TaskScheduler, UMovieSceneEntitySystemLinker* Linker, const FEntityComponentFilter& AdditionalFilter)
 	{
 		FBuiltInComponentTypes* BuiltInComponents = FBuiltInComponentTypes::Get();
 
@@ -377,6 +382,7 @@ struct TPropertyComponentHandlerImpl<PropertyTraits, TPropertyMetaData<MetaDataT
 		.ReadAllOf(Composites[CompositeIndices].ComponentTypeID.ReinterpretCast<CompositeTypes>()...)
 		.FilterAll({ Definition.PropertyType })
 		.FilterNone({ BuiltInComponents->Tags.Ignored })
+		.CombineFilter(AdditionalFilter)
 		.SetStat(Definition.StatID)
 		.SetDesiredThread(Linker->EntityManager.GetGatherThread())
 		.template Fork_PerAllocation<CompleteSetterTask>(&Linker->EntityManager, TaskScheduler, Definition.CustomPropertyRegistration);
@@ -401,6 +407,7 @@ struct TPropertyComponentHandlerImpl<PropertyTraits, TPropertyMetaData<MetaDataT
 				.FilterAny({ CompletePropertyMask })
 				.FilterAll({ Definition.PropertyType })
 				.FilterNone({ BuiltInComponents->Tags.Ignored })
+				.CombineFilter(AdditionalFilter)
 				.FilterOut(CompletePropertyMask)
 				.SetStat(Definition.StatID)
 				.SetDesiredThread(Linker->EntityManager.GetGatherThread())
diff --git a/Engine/Source/Runtime/MovieSceneTracks/Private/MovieSceneTracksComponentTypes.cpp b/Engine/Source/Runtime/MovieSceneTracks/Private/MovieSceneTracksComponentTypes.cpp
--- a/Engine/Source/Runtime/MovieSceneTracks/Private/MovieSceneTracksComponentTypes.cpp
+++ b/Engine/Source/Runtime/MovieSceneTracks/Private/MovieSceneTracksComponentTypes.cpp
@@ -506,6 +506,10 @@ struct FComponentTransformHandler : TPropertyComponentHandler<FComponentTransfor
 	{
 		return Container->GetOrCreateStorage<FPreAnimatedComponentTransformStorage>();
 	}
+	virtual void ScheduleSetterTasks(const FPropertyDefinition& Definition, TArrayView<const FPropertyCompositeDefinition> Composites, const FPropertyStats& Stats, IEntitySystemScheduler* TaskScheduler, UMovieSceneEntitySystemLinker* Linker) override
+	{
+		ScheduleSetterTasksImpl(Definition, Composites, Stats, TaskScheduler, Linker, FEntityComponentFilter().None({ FMovieSceneTracksComponentTypes::Get()->Tags.AnimMixerPoseProducer }));
+	}
 };
 
 struct FObjectHandler : TPropertyComponentHandler<FObjectPropertyTraits, FObjectComponent>
