Fix crash when clearing anim class when a sequence is using the Custom Anim Instance Mixed Animation Target.
-Adding event to the skeletal mesh component that is broadcast when the component is unregistered, so we can unbind the target anim instance from the mesh.

diff --git a/Engine/Plugins/Experimental/MovieSceneAnimMixer/Source/MovieSceneAnimMixer/Private/Systems/MovieSceneAnimInstanceTargetSystem.cpp b/Engine/Plugins/Experimental/MovieSceneAnimMixer/Source/MovieSceneAnimMixer/Private/Systems/MovieSceneAnimInstanceTargetSystem.cpp
--- a/Engine/Plugins/Experimental/MovieSceneAnimMixer/Source/MovieSceneAnimMixer/Private/Systems/MovieSceneAnimInstanceTargetSystem.cpp
+++ b/Engine/Plugins/Experimental/MovieSceneAnimMixer/Source/MovieSceneAnimMixer/Private/Systems/MovieSceneAnimInstanceTargetSystem.cpp
@@ -475,15 +475,26 @@ namespace UE::MovieScene
 #if WITH_EDITOR
 			if (GEditor && bWasCreated)
 			{
-				FDelegateHandle NewDelegateHandle = GEditor->OnBlueprintPreCompile().AddLambda([CurrentAnimInstanceWeakPtr, SkeletalMeshComponentWeakPtr](const UBlueprint* InBlueprint)
+				FDelegateHandle PreCompileHandle = GEditor->OnBlueprintPreCompile().AddLambda([CurrentAnimInstanceWeakPtr, SkeletalMeshComponentWeakPtr](const UBlueprint* InBlueprint)
 				{
 					if (CurrentAnimInstanceWeakPtr.IsValid() && SkeletalMeshComponentWeakPtr.IsValid() && InBlueprint && CurrentAnimInstanceWeakPtr->GetClass() == InBlueprint->GeneratedClass)
 					{
 						FAnimCustomInstanceHelper::UnbindFromSkeletalMeshComponent<USequencerMixedAnimInstance>(SkeletalMeshComponentWeakPtr.Get());
 					}
 				});
+				
+				auto UnregisteredLambda = USkeletalMeshComponent::FOnSkeletalMeshUnregisteredDelegate::CreateLambda([CurrentAnimInstanceWeakPtr, SkeletalMeshComponentWeakPtr](const USkeletalMeshComponent* InMeshComponent)
+				{
+					if (CurrentAnimInstanceWeakPtr.IsValid() && SkeletalMeshComponentWeakPtr.IsValid() && InMeshComponent == SkeletalMeshComponentWeakPtr.Get())
+					{
+						FAnimCustomInstanceHelper::UnbindFromSkeletalMeshComponent<USequencerMixedAnimInstance>(SkeletalMeshComponentWeakPtr.Get());
+					}
+				});
+
+				FDelegateHandle SkeletalMeshUnregisteredHandle = SkeletalMeshComponent->RegisterOnSkeletalMeshUnregistered(UnregisteredLambda);
 			
-				System->Handles.Add(NewDelegateHandle);
+				System->PreCompileHandles.Add(PreCompileHandle);
+				System->MeshUnregisteredHandles.Add({SkeletalMeshUnregisteredHandle, SkeletalMeshComponentWeakPtr});
 			}
 #endif
 			
@@ -543,10 +554,18 @@ UMovieSceneAnimInstanceTargetSystem::~UMovieSceneAnimInstanceTargetSystem()
 {
 	if (GEditor)
 	{
-		for (FDelegateHandle Handle : Handles)
+		for (FDelegateHandle Handle : PreCompileHandles)
 		{
 			GEditor->OnBlueprintPreCompile().Remove(Handle);
 		}
+
+		for (TTuple<FDelegateHandle, TWeakObjectPtr<USkeletalMeshComponent>> HandlePair : MeshUnregisteredHandles)
+		{
+			if (TStrongObjectPtr<USkeletalMeshComponent> SkeletalMeshComponent = HandlePair.Value.Pin())
+			{
+				SkeletalMeshComponent->UnregisterOnSkeletalMeshUnregistered(HandlePair.Key);
+			}
+		}
 	}
 }
 #endif
diff --git a/Engine/Plugins/Experimental/MovieSceneAnimMixer/Source/MovieSceneAnimMixer/Public/Systems/MovieSceneAnimInstanceTargetSystem.h b/Engine/Plugins/Experimental/MovieSceneAnimMixer/Source/MovieSceneAnimMixer/Public/Systems/MovieSceneAnimInstanceTargetSystem.h
--- a/Engine/Plugins/Experimental/MovieSceneAnimMixer/Source/MovieSceneAnimMixer/Public/Systems/MovieSceneAnimInstanceTargetSystem.h
+++ b/Engine/Plugins/Experimental/MovieSceneAnimMixer/Source/MovieSceneAnimMixer/Public/Systems/MovieSceneAnimInstanceTargetSystem.h
@@ -123,7 +123,8 @@ public:
 #if WITH_EDITOR
 	virtual ~UMovieSceneAnimInstanceTargetSystem() override;
 
-	TArray<FDelegateHandle> Handles;
+	TArray<FDelegateHandle> PreCompileHandles;
+	TArray<TTuple<FDelegateHandle, TWeakObjectPtr<USkeletalMeshComponent>>> MeshUnregisteredHandles;
 #endif
 	
 private:
diff --git a/Engine/Source/Runtime/Engine/Classes/Components/SkeletalMeshComponent.h b/Engine/Source/Runtime/Engine/Classes/Components/SkeletalMeshComponent.h
--- a/Engine/Source/Runtime/Engine/Classes/Components/SkeletalMeshComponent.h
+++ b/Engine/Source/Runtime/Engine/Classes/Components/SkeletalMeshComponent.h
@@ -1779,6 +1779,14 @@ public:
 	ENGINE_API FDelegateHandle RegisterOnSkeletalMeshPropertyChanged(const FOnSkeletalMeshPropertyChanged& Delegate);
 	ENGINE_API void UnregisterOnSkeletalMeshPropertyChanged(FDelegateHandle Handle);
 
+	DECLARE_MULTICAST_DELEGATE_OneParam(FOnSkeletalMeshUnregisterd, const USkeletalMeshComponent*)
+	FOnSkeletalMeshUnregisterd OnSkeletalMeshUnregistered;
+	typedef FOnSkeletalMeshUnregisterd::FDelegate FOnSkeletalMeshUnregisteredDelegate;
+
+	/** Register / Unregister delegates called when the mesh component is unregistered. */
+	ENGINE_API FDelegateHandle RegisterOnSkeletalMeshUnregistered(const FOnSkeletalMeshUnregisteredDelegate& Delegate);
+	ENGINE_API void UnregisterOnSkeletalMeshUnregistered(FDelegateHandle Handle);
+
 	ENGINE_API virtual void PostEditChangeProperty(FPropertyChangedEvent& PropertyChangedEvent) override;
 
 	/** Validates the animation asset or blueprint, making sure it is compatible with the current skeleton */
diff --git a/Engine/Source/Runtime/Engine/Private/Components/SkeletalMeshComponent.cpp b/Engine/Source/Runtime/Engine/Private/Components/SkeletalMeshComponent.cpp
--- a/Engine/Source/Runtime/Engine/Private/Components/SkeletalMeshComponent.cpp
+++ b/Engine/Source/Runtime/Engine/Private/Components/SkeletalMeshComponent.cpp
@@ -1058,6 +1058,10 @@ void USkeletalMeshComponent::OnUnregister()
 	//clothing actors will be re-created in TickClothing
 	ReleaseAllClothingResources();
 
+#if WITH_EDITOR
+	OnSkeletalMeshUnregistered.Broadcast(this);
+#endif
+	
 	if (AnimScriptInstance)
 	{
 		AnimScriptInstance->UninitializeAnimation();
@@ -1412,6 +1416,16 @@ void USkeletalMeshComponent::BeginPlay()
 }
 
 #if WITH_EDITOR
+FDelegateHandle USkeletalMeshComponent::RegisterOnSkeletalMeshUnregistered(const FOnSkeletalMeshUnregisteredDelegate& Delegate)
+{
+	return OnSkeletalMeshUnregistered.Add(Delegate);
+}
+
+void USkeletalMeshComponent::UnregisterOnSkeletalMeshUnregistered(FDelegateHandle Handle)
+{
+	OnSkeletalMeshUnregistered.Remove(Handle);
+}
+
 void USkeletalMeshComponent::PostEditChangeProperty(FPropertyChangedEvent& PropertyChangedEvent)
 {
 	FProperty* PropertyThatChanged = PropertyChangedEvent.Property;
