fix compile: error: duplicate symbol: UE::MovieScene::GetAnimSequencerInstance(USkeletalMeshComponent*)

diff --git a/Engine/Plugins/Experimental/MovieSceneAnimMixer/Source/MovieSceneAnimMixer/Private/Systems/MovieSceneAnimInstanceTargetSystem.cpp b/Engine/Plugins/Experimental/MovieSceneAnimMixer/Source/MovieSceneAnimMixer/Private/Systems/MovieSceneAnimInstanceTargetSystem.cpp
--- a/Engine/Plugins/Experimental/MovieSceneAnimMixer/Source/MovieSceneAnimMixer/Private/Systems/MovieSceneAnimInstanceTargetSystem.cpp
+++ b/Engine/Plugins/Experimental/MovieSceneAnimMixer/Source/MovieSceneAnimMixer/Private/Systems/MovieSceneAnimInstanceTargetSystem.cpp
@@ -290,12 +290,6 @@ void FSequencerMixedAnimInstanceProxy::SetMixerTask(TSharedPtr<FAnimNextEvaluati
 
 namespace UE::MovieScene
 {
-	/** Helper function to get our sequencer animation instance from a skeletal mesh component */
-	USequencerMixedAnimInstance* GetAnimSequencerInstance(USkeletalMeshComponent* SkeletalMeshComponent)
-	{
-		return Cast<USequencerMixedAnimInstance>(SkeletalMeshComponent->GetAnimInstance());
-	}
-
 	/** ------------------------------------------------------------------------- */
 
 	/** Pre-animated state for skeletal animations */
