Sequencer: Added null check against invalid root motion ptr

diff --git a/Engine/Plugins/Experimental/MovieSceneAnimMixer/Source/MovieSceneAnimMixer/Private/Systems/MovieSceneRootMotionSystem.cpp b/Engine/Plugins/Experimental/MovieSceneAnimMixer/Source/MovieSceneAnimMixer/Private/Systems/MovieSceneRootMotionSystem.cpp
--- a/Engine/Plugins/Experimental/MovieSceneAnimMixer/Source/MovieSceneAnimMixer/Private/Systems/MovieSceneRootMotionSystem.cpp
+++ b/Engine/Plugins/Experimental/MovieSceneAnimMixer/Source/MovieSceneAnimMixer/Private/Systems/MovieSceneRootMotionSystem.cpp
@@ -55,7 +55,10 @@ namespace UE::MovieScene
 			)
 		{
 			OutRootDestination = (EMovieSceneRootMotionDestination)(ByteResult);
-			RootMotion->RootDestination = OutRootDestination;
+			if (ensure(RootMotion))
+			{
+				RootMotion->RootDestination = OutRootDestination;
+			}
 		}
 	};
 
