Rename Anim Next Module Injection to UAF Module Injection

diff --git a/Engine/Plugins/Experimental/MovieSceneAnimMixer/Source/MovieSceneAnimMixer/Public/Systems/MovieSceneAnimNextTargetSystem.h b/Engine/Plugins/Experimental/MovieSceneAnimMixer/Source/MovieSceneAnimMixer/Public/Systems/MovieSceneAnimNextTargetSystem.h
--- a/Engine/Plugins/Experimental/MovieSceneAnimMixer/Source/MovieSceneAnimMixer/Public/Systems/MovieSceneAnimNextTargetSystem.h
+++ b/Engine/Plugins/Experimental/MovieSceneAnimMixer/Source/MovieSceneAnimMixer/Public/Systems/MovieSceneAnimNextTargetSystem.h
@@ -30,9 +30,9 @@ struct FMovieSceneAnimMixerEvaluationModifier : public UE::AnimNext::IEvaluation
 };
 
 /**
- * Declaring a unique target for targeting an injection site on an Anim Next module. Will find the default injection site, or one specified by name.
+ * Declaring a unique target for targeting an injection site on an Unreal Animation Framework module. Will find the default injection site, or one specified by name.
  */
-USTRUCT(meta=(DisplayName="Anim Next Module Injection"))
+USTRUCT(meta=(DisplayName="UAF Module Injection"))
 struct MOVIESCENEANIMMIXER_API FMovieSceneAnimNextInjectionTarget : public FMovieSceneMixedAnimationTarget
 {
 	GENERATED_BODY() 
@@ -56,7 +56,7 @@ struct FMovieSceneAnimNextTargetData
 	TSharedPtr<FMovieSceneAnimMixerEvaluationModifier> Modifier = nullptr;
 };
 
-// System that handles applying animation mixer evaluation tasks to an injection site in an Anim Next module
+// System that handles applying animation mixer evaluation tasks to an injection site in an Unreal Animation Framework module
 UCLASS(MinimalAPI)
 class UMovieSceneAnimNextTargetSystem
 	: public UMovieSceneEntitySystem
