Sequencer: Fixed ComponentToActorTransform using the previous frame's actor transform when calculating root motion

FInitializeRootMotionTask::ForEachEntity can set actor transforms when root motion is not being applied to the actor. Since we were previously calculating the ComponentToActor transform using the ComponentToWorld transform for both the actor and the component before unning this task, it was possible that we calculated the transform using the previous frame's data. In order to correct this, we now calculate these transforms _after_ applying all the other transforms touched by ForEachEntity, making sure not to overwrite any ActorTransform members that were previously set by the task.

diff --git a/Engine/Plugins/Experimental/MovieSceneAnimMixer/Source/MovieSceneAnimMixer/Private/Systems/MovieSceneAnimMixerSystem.cpp b/Engine/Plugins/Experimental/MovieSceneAnimMixer/Source/MovieSceneAnimMixer/Private/Systems/MovieSceneAnimMixerSystem.cpp
--- a/Engine/Plugins/Experimental/MovieSceneAnimMixer/Source/MovieSceneAnimMixer/Private/Systems/MovieSceneAnimMixerSystem.cpp
+++ b/Engine/Plugins/Experimental/MovieSceneAnimMixer/Source/MovieSceneAnimMixer/Private/Systems/MovieSceneAnimMixerSystem.cpp
@@ -586,6 +586,18 @@ void UMovieSceneAnimMixerSystem::AssignRootMotion(FObjectKey InObject, TSharedPt
 	}
 }
 
+void UMovieSceneAnimMixerSystem::PreInitializeAllRootMotion()
+{
+	for (auto It = RootMotionByObject.CreateIterator(); It; ++It)
+	{
+		TSharedPtr<FMovieSceneMixerRootMotionComponentData> RootMotion = It.Value().Pin();
+		if (RootMotion)
+		{
+			RootMotion->bActorTransformSet = false;
+		}
+	}
+}
+
 void UMovieSceneAnimMixerSystem::InitializeAllRootMotion()
 {
 	for (auto It = RootMotionByObject.CreateIterator(); It; ++It)
diff --git a/Engine/Plugins/Experimental/MovieSceneAnimMixer/Source/MovieSceneAnimMixer/Private/Systems/MovieSceneRootMotionSystem.cpp b/Engine/Plugins/Experimental/MovieSceneAnimMixer/Source/MovieSceneAnimMixer/Private/Systems/MovieSceneRootMotionSystem.cpp
--- a/Engine/Plugins/Experimental/MovieSceneAnimMixer/Source/MovieSceneAnimMixer/Private/Systems/MovieSceneRootMotionSystem.cpp
+++ b/Engine/Plugins/Experimental/MovieSceneAnimMixer/Source/MovieSceneAnimMixer/Private/Systems/MovieSceneRootMotionSystem.cpp
@@ -72,7 +72,7 @@ namespace UE::MovieScene
 
 		void PreTask()
 		{
-			AnimMixer->InitializeAllRootMotion();
+			AnimMixer->PreInitializeAllRootMotion();
 		}
 
 		void ForEachEntity(UObject* BoundObject,
@@ -111,6 +111,7 @@ namespace UE::MovieScene
 					if (InScaleZ) { Scale.Z = *InScaleZ; }
 
 					RootMotion->ActorTransform = FTransform(Rotation, Location, Scale);
+					RootMotion->bActorTransformSet = true;
 					return;
 				}
 			}
@@ -132,6 +133,11 @@ namespace UE::MovieScene
 
 			UE::MovieScene::FIntermediate3DTransform::ApplyTransformTo(SceneComponent, IntermediateTransform);
 		}
+
+		void PostTask()
+		{
+			AnimMixer->InitializeAllRootMotion();
+		}
 	};
 
 	struct FIsActorBeingMovedStatics
@@ -216,7 +222,10 @@ void FMovieSceneMixerRootMotionComponentData::Initialize()
 			InverseMeshToActorRotation = RootComponent->GetComponentTransform().GetRelativeTransformReverse(BoundObject->GetComponentTransform()).GetRotation();
 		}
 
-		ActorTransform = RootComponent->GetRelativeTransform();
+		if (!bActorTransformSet)
+		{
+			ActorTransform = RootComponent->GetRelativeTransform();
+		}
 	}
 
 	switch (RootDestination)
diff --git a/Engine/Plugins/Experimental/MovieSceneAnimMixer/Source/MovieSceneAnimMixer/Public/Systems/MovieSceneAnimMixerSystem.h b/Engine/Plugins/Experimental/MovieSceneAnimMixer/Source/MovieSceneAnimMixer/Public/Systems/MovieSceneAnimMixerSystem.h
--- a/Engine/Plugins/Experimental/MovieSceneAnimMixer/Source/MovieSceneAnimMixer/Public/Systems/MovieSceneAnimMixerSystem.h
+++ b/Engine/Plugins/Experimental/MovieSceneAnimMixer/Source/MovieSceneAnimMixer/Public/Systems/MovieSceneAnimMixerSystem.h
@@ -153,6 +153,7 @@ public:
 	TSharedPtr<FMovieSceneMixerRootMotionComponentData> FindRootMotion(FObjectKey InObject) const;
 	void AssignRootMotion(FObjectKey InObjectKey, TSharedPtr<FMovieSceneMixerRootMotionComponentData> RootMotion);
 
+	void PreInitializeAllRootMotion();
 	void InitializeAllRootMotion();
 
 private:
diff --git a/Engine/Plugins/Experimental/MovieSceneAnimMixer/Source/MovieSceneAnimMixer/Public/Systems/MovieSceneRootMotionSystem.h b/Engine/Plugins/Experimental/MovieSceneAnimMixer/Source/MovieSceneAnimMixer/Public/Systems/MovieSceneRootMotionSystem.h
--- a/Engine/Plugins/Experimental/MovieSceneAnimMixer/Source/MovieSceneAnimMixer/Public/Systems/MovieSceneRootMotionSystem.h
+++ b/Engine/Plugins/Experimental/MovieSceneAnimMixer/Source/MovieSceneAnimMixer/Public/Systems/MovieSceneRootMotionSystem.h
@@ -89,6 +89,7 @@ private:
 public:
 
 	bool bComponentSpaceRoot = false;
+	bool bActorTransformSet = false;
 };
 
 // Task that converts the root motion attribute on the top pose of the pose stack to world space by adding on the actor transformation, root base transform, and/or transform origin.
