Sequencer: CachedAnimInstance to WeakObjectPtr to fix GC errors

diff --git a/Engine/Plugins/Experimental/MovieSceneAnimMixer/Source/MovieSceneAnimMixer/Private/Systems/MovieSceneAnimInstanceTargetSystem.cpp b/Engine/Plugins/Experimental/MovieSceneAnimMixer/Source/MovieSceneAnimMixer/Private/Systems/MovieSceneAnimInstanceTargetSystem.cpp
--- a/Engine/Plugins/Experimental/MovieSceneAnimMixer/Source/MovieSceneAnimMixer/Private/Systems/MovieSceneAnimInstanceTargetSystem.cpp
+++ b/Engine/Plugins/Experimental/MovieSceneAnimMixer/Source/MovieSceneAnimMixer/Private/Systems/MovieSceneAnimInstanceTargetSystem.cpp
@@ -296,7 +296,7 @@ namespace UE::MovieScene
 	struct FPreAnimatedMixedSkeletalAnimationState
 	{
 		EAnimationMode::Type AnimationMode;
-		TStrongObjectPtr<UAnimInstance> CachedAnimInstance;
+		TWeakObjectPtr<UAnimInstance> CachedAnimInstance;
 		FSkeletalMeshRestoreState SkeletalMeshRestoreState;
 	};
 
@@ -313,7 +313,7 @@ namespace UE::MovieScene
 			if (ensure(Component))
 			{
 				OutCachedValue.AnimationMode = Component->GetAnimationMode();
-				OutCachedValue.CachedAnimInstance.Reset(Component->AnimScriptInstance);
+				OutCachedValue.CachedAnimInstance = Component->AnimScriptInstance;
 				OutCachedValue.SkeletalMeshRestoreState.SaveState(Component);
 			}
 			return OutCachedValue;
@@ -338,10 +338,11 @@ namespace UE::MovieScene
 				// if we're using same anim blueprint, we don't want to keep reinitializing it. 
 				Component->SetAnimationMode(InOutCachedValue.AnimationMode);
 			}
-			UAnimInstance* PreviousAnimInstance = InOutCachedValue.CachedAnimInstance.Get();
-			if (PreviousAnimInstance && IsValid(PreviousAnimInstance) && PreviousAnimInstance->GetSkelMeshComponent() == Component)
+
+			TStrongObjectPtr<UAnimInstance> PreviousAnimInstance = InOutCachedValue.CachedAnimInstance.Pin();
+			if (PreviousAnimInstance && IsValid(PreviousAnimInstance.Get()) && PreviousAnimInstance->GetSkelMeshComponent() == Component)
 			{
-				Component->AnimScriptInstance = PreviousAnimInstance;
+				Component->AnimScriptInstance = PreviousAnimInstance.Get();
 				InOutCachedValue.CachedAnimInstance.Reset();
 				if (Component->AnimScriptInstance && Component->GetSkeletalMeshAsset() && Component->AnimScriptInstance->CurrentSkeleton != Component->GetSkeletalMeshAsset()->GetSkeleton())
 				{
diff --git a/Engine/Source/Runtime/MovieSceneTracks/Private/Systems/MovieSceneSkeletalAnimationSystem.cpp b/Engine/Source/Runtime/MovieSceneTracks/Private/Systems/MovieSceneSkeletalAnimationSystem.cpp
--- a/Engine/Source/Runtime/MovieSceneTracks/Private/Systems/MovieSceneSkeletalAnimationSystem.cpp
+++ b/Engine/Source/Runtime/MovieSceneTracks/Private/Systems/MovieSceneSkeletalAnimationSystem.cpp
@@ -91,7 +91,7 @@ UAnimSequencerInstance* GetAnimSequencerInstance(USkeletalMeshComponent* Skeleta
 struct FPreAnimatedSkeletalAnimationState
 {
 	EAnimationMode::Type AnimationMode;
-	TStrongObjectPtr<UAnimInstance> CachedAnimInstance;
+	TWeakObjectPtr<UAnimInstance> CachedAnimInstance;
 	FSkeletalMeshRestoreState SkeletalMeshRestoreState;
 };
 
@@ -108,7 +108,7 @@ struct FPreAnimatedSkeletalAnimationTraits : FBoundObjectPreAnimatedStateTraits
 		if (ensure(Component))
 		{
 			OutCachedValue.AnimationMode = Component->GetAnimationMode();
-			OutCachedValue.CachedAnimInstance.Reset(Component->AnimScriptInstance);
+			OutCachedValue.CachedAnimInstance = Component->AnimScriptInstance;
 			OutCachedValue.SkeletalMeshRestoreState.SaveState(Component);
 		}
 		return OutCachedValue;
@@ -141,10 +141,10 @@ struct FPreAnimatedSkeletalAnimationTraits : FBoundObjectPreAnimatedStateTraits
 			Component->SetAnimationMode(InOutCachedValue.AnimationMode);
 		}
 
-		UAnimInstance* PreviousAnimInstance = InOutCachedValue.CachedAnimInstance.Get();
-		if (PreviousAnimInstance && PreviousAnimInstance->GetSkelMeshComponent() == Component)
+		TStrongObjectPtr<UAnimInstance> PreviousAnimInstance = InOutCachedValue.CachedAnimInstance.Pin();
+		if (PreviousAnimInstance && IsValid(PreviousAnimInstance.Get()) && PreviousAnimInstance->GetSkelMeshComponent() == Component)
 		{
-			Component->AnimScriptInstance = PreviousAnimInstance;
+			Component->AnimScriptInstance = PreviousAnimInstance.Get();
 			InOutCachedValue.CachedAnimInstance.Reset();
 			if (Component->AnimScriptInstance && Component->GetSkeletalMeshAsset() && Component->AnimScriptInstance->CurrentSkeleton != Component->GetSkeletalMeshAsset()->GetSkeleton())
 			{
