Imporove default anim mixer workflow when using UAF.

The default injection site for level sequences can now be set in the project settings, instead of only through a cvar.
The anim mixer will now create an anim next component with a default module, if it found one that does not have a valid module, to support a workflow similar to the custom mode for anim blueprints.
The default module can be specified in the project settings, and defaults to a minimal module that has a simple injection setup.

#virtualized

diff --git a/Engine/Config/BaseEngine.ini b/Engine/Config/BaseEngine.ini
--- a/Engine/Config/BaseEngine.ini
+++ b/Engine/Config/BaseEngine.ini
@@ -3683,6 +3683,9 @@ DefaultCompletionMode=RestoreState
 [/Script/TemplateSequence.TemplateSequence]
 DefaultCompletionMode=RestoreState
 
+[/Script/MovieSceneAnimMixer.MovieSceneAnimMixerSettings]
+DefaultUAFModule=/MovieSceneAnimMixer/DefaultAnimMixerUAFModule.DefaultAnimMixerUAFModule
+
 [PlatformCrypto]
 PlatformRequiresDataCrypto=False
 PakSigningRequired=True
diff --git a/Engine/Plugins/Experimental/MovieSceneAnimMixer/Source/MovieSceneAnimMixer/MovieSceneAnimMixer.Build.cs b/Engine/Plugins/Experimental/MovieSceneAnimMixer/Source/MovieSceneAnimMixer/MovieSceneAnimMixer.Build.cs
--- a/Engine/Plugins/Experimental/MovieSceneAnimMixer/Source/MovieSceneAnimMixer/MovieSceneAnimMixer.Build.cs
+++ b/Engine/Plugins/Experimental/MovieSceneAnimMixer/Source/MovieSceneAnimMixer/MovieSceneAnimMixer.Build.cs
@@ -12,6 +12,8 @@ namespace UnrealBuildTool.Rules
 					"Core",
 					"CoreUObject",
 					"Engine",
+					"DeveloperSettings",
+					"Settings",
 					"MovieScene",
 					"MovieSceneTracks",
 					"AnimGraphRuntime",
diff --git a/Engine/Plugins/Experimental/MovieSceneAnimMixer/Source/MovieSceneAnimMixer/Private/MovieSceneAnimMixerModule.cpp b/Engine/Plugins/Experimental/MovieSceneAnimMixer/Source/MovieSceneAnimMixer/Private/MovieSceneAnimMixerModule.cpp
--- a/Engine/Plugins/Experimental/MovieSceneAnimMixer/Source/MovieSceneAnimMixer/Private/MovieSceneAnimMixerModule.cpp
+++ b/Engine/Plugins/Experimental/MovieSceneAnimMixer/Source/MovieSceneAnimMixer/Private/MovieSceneAnimMixerModule.cpp
@@ -1,16 +1,38 @@
 // Copyright Epic Games, Inc. All Rights Reserved.
 
 #include "MovieSceneAnimMixerModule.h"
+#include "MovieSceneAnimMixerSettings.h"
+#include "ISettingsModule.h"
 #include "Modules/ModuleInterface.h"
 #include "Modules/ModuleManager.h"
 
 DEFINE_LOG_CATEGORY(LogMovieSceneAnimMixer);
 
+#define LOCTEXT_NAMESPACE "FAnimMixerModule"
+
 namespace UE::MovieScene
 {
+	FName GSequencerDefaultAnimNextInjectionSite = NAME_None;
+
+	FAutoConsoleVariableRef CVarDefaultAnimNextInjectionSite(
+	TEXT("Sequencer.AnimNext.DefaultInjectionSite"),
+	GSequencerDefaultAnimNextInjectionSite,
+	TEXT("(Default: None) Specifies the default injection site name for Sequencer Anim Next Targets that is used when none is specified on the target itself."),
+	ECVF_Default
+	);
 
 void FMovieSceneAnimMixerModule::StartupModule()
 {
+	ISettingsModule& SettingsModule = FModuleManager::Get().LoadModuleChecked<ISettingsModule>("Settings");
+	
+	SettingsModule.RegisterSettings("Plugins", "Animation", "AnimMixer",
+		LOCTEXT("RuntimeSettingsName", "Anim Mixer"),
+		LOCTEXT("RuntimeSettingsDescription", "Configure project settings relating to the Anim Mixer Plugin"),
+		GetMutableDefault<UMovieSceneAnimMixerSettings>()
+	);
+
+	CVarDefaultAnimNextInjectionSite->Set(*GetMutableDefault<UMovieSceneAnimMixerSettings>()->DefaultInjectionSite.ToString(), ECVF_SetByProjectSetting);
+	
 }
 
 void FMovieSceneAnimMixerModule::ShutdownModule()
@@ -21,3 +43,5 @@ void FMovieSceneAnimMixerModule::ShutdownModule()
 }
 
 IMPLEMENT_MODULE(UE::MovieScene::FMovieSceneAnimMixerModule, MovieSceneAnimMixer)
+
+#undef LOCTEXT_NAMESPACE
\ No newline at end of file
diff --git a/Engine/Plugins/Experimental/MovieSceneAnimMixer/Source/MovieSceneAnimMixer/Private/MovieSceneAnimMixerSettings.cpp b/Engine/Plugins/Experimental/MovieSceneAnimMixer/Source/MovieSceneAnimMixer/Private/MovieSceneAnimMixerSettings.cpp
new file mode 100644
--- /dev/null
+++ b/Engine/Plugins/Experimental/MovieSceneAnimMixer/Source/MovieSceneAnimMixer/Private/MovieSceneAnimMixerSettings.cpp
@@ -0,0 +1,28 @@
+﻿// Copyright Epic Games, Inc. All Rights Reserved.
+
+#include "MovieSceneAnimMixerSettings.h"
+
+#if WITH_EDITOR
+#include "MovieSceneAnimMixerModule.h"
+#endif
+
+FName UMovieSceneAnimMixerSettings::GetCategoryName() const
+{
+	return TEXT("Plugins");
+}
+
+
+
+#if WITH_EDITOR
+void UMovieSceneAnimMixerSettings::PostEditChangeProperty(FPropertyChangedEvent& PropertyChangedEvent)
+{
+	Super::PostEditChangeProperty(PropertyChangedEvent);
+
+	UE::MovieScene::CVarDefaultAnimNextInjectionSite->Set(*DefaultInjectionSite.ToString(), ECVF_SetByProjectSetting);
+}
+
+FText UMovieSceneAnimMixerSettings::GetSectionText() const
+{
+	return Super::GetSectionText();
+}
+#endif
\ No newline at end of file
diff --git a/Engine/Plugins/Experimental/MovieSceneAnimMixer/Source/MovieSceneAnimMixer/Private/Systems/MovieSceneAnimNextTargetSystem.cpp b/Engine/Plugins/Experimental/MovieSceneAnimMixer/Source/MovieSceneAnimMixer/Private/Systems/MovieSceneAnimNextTargetSystem.cpp
--- a/Engine/Plugins/Experimental/MovieSceneAnimMixer/Source/MovieSceneAnimMixer/Private/Systems/MovieSceneAnimNextTargetSystem.cpp
+++ b/Engine/Plugins/Experimental/MovieSceneAnimMixer/Source/MovieSceneAnimMixer/Private/Systems/MovieSceneAnimNextTargetSystem.cpp
@@ -10,6 +10,7 @@
 #include "EntitySystem/MovieSceneEntitySystemTask.h"
 #include "MovieSceneTracksComponentTypes.h"
 #include "SkeletalMeshRestoreState.h"
+#include "MovieSceneAnimMixerModule.h"
 
 #include "Evaluation/PreAnimatedState/IMovieScenePreAnimatedStorage.h"
 #include "Evaluation/PreAnimatedState/MovieScenePreAnimatedObjectStorage.h"
@@ -31,25 +32,22 @@
 #include "Animation/BuiltInAttributeTypes.h"
 #include "Systems/MovieSceneAnimMixerSystem.h"
 #include "AnimMixerComponentTypes.h"
+#include "MovieSceneAnimMixerSettings.h"
+#include "AnimMixerUtils.h"
 
 #include "Graph/AnimNext_LODPose.h"
 #include "Injection/InjectionUtils.h"
 #include "TraitInterfaces/IEvaluate.h"
 
 #include "PrimitiveSceneProxy.h"
+#include "Module/AnimNextModule.h"
+#include "UAF/Private/Component/AnimNextComponentWorldSubsystem.h"
 #include "UObject/FortniteMainBranchObjectVersion.h"
 
 #include UE_INLINE_GENERATED_CPP_BY_NAME(MovieSceneAnimNextTargetSystem)
 
 namespace UE::MovieScene
 {
-	FName GSequencerDefaultAnimNextInjectionSite = NAME_None;
-	FAutoConsoleVariableRef CVarDefaultAnimNextInjectionSite(
-		TEXT("Sequencer.AnimNext.DefaultInjectionSite"),
-		GSequencerDefaultAnimNextInjectionSite,
-		TEXT("(Default: None) Specifies the default injection site name for Sequencer Anim Next Targets that is used when none is specified on the target itself."),
-		ECVF_Default
-	);
 
 	/** Pre-animated state for skeletal animations */
 	struct FPreAnimatedAnimNextState
@@ -58,9 +56,15 @@ namespace UE::MovieScene
 		UE::UAF::FInjectionRequestPtr InjectionHandle;
 
 		TObjectPtr<UMovieSceneAnimNextTargetSystem> System;
+	
+		// Pointer to the UAF Component's module. Since the system sets the value on the component, this is a strong ptr to prevent GC. 	
+		TStrongObjectPtr<UAnimNextModule> PreAnimatedModule;
 
 		// Whether the anim next module started enabled or not.
 		bool bAnimNextModuleEnabled = false;
+
+		// Whether evaluation changed the module, and therefore needs to be reset.
+		bool bModuleReplaced = false;
 	};
 
 	/** Pre-animation traits for skeletal animations */
@@ -81,6 +85,19 @@ namespace UE::MovieScene
 				UE::UAF::FInjectionUtils::Uninject(InOutCachedValue.InjectionHandle);
 			}
 
+			if (InOutCachedValue.bModuleReplaced)
+			{
+				if (UAnimNextComponent* AnimNextComponent = Cast<UAnimNextComponent>(ObjectAndInjectionSite.Key.ResolveObjectPtr()))
+				{
+					// TODO the animation needs to be updated, or the post needs to be reset to ref pose when restoring to a previous module.
+					// Currently there is no support for this in UAF.
+					// This is most pronounced in the editor, where the mesh stays in the same pose when the sequence is closed.
+					FAnimMixerUtils::UnregisterWithUAFSubsystem(AnimNextComponent);
+					FAnimMixerUtils::SetUAFComponentModule(AnimNextComponent, InOutCachedValue.PreAnimatedModule.Get());
+					FAnimMixerUtils::RegisterWithUAFSubsystem(AnimNextComponent);
+				}
+			}
+			
 			// Reset the anim next module back to not being enabled if applicable
 			if (!InOutCachedValue.bAnimNextModuleEnabled)
 			{
@@ -222,6 +239,7 @@ namespace UE::MovieScene
 			UObject* BoundMesh = ObjectComponent.GetObject();
 			USkeletalMeshComponent* SkeletalMeshComponent = Cast<USkeletalMeshComponent>(BoundMesh);
 			UAnimNextComponent* AnimNextComponent = Cast<UAnimNextComponent>(BoundMesh);
+			bool bModuleReplaced = false;
 
 			if (!AnimNextComponent)
 			{
@@ -244,6 +262,18 @@ namespace UE::MovieScene
 
 			if (AnimNextComponent)
 			{
+				TObjectPtr<UAnimNextModule> PreAnimatedModulePtr = FAnimMixerUtils::GetUAFModule(AnimNextComponent);
+				TStrongObjectPtr<UAnimNextModule> PreAnimatedModule = TStrongObjectPtr<UAnimNextModule>(FAnimMixerUtils::GetUAFModule(AnimNextComponent));
+				if (!FAnimMixerUtils::IsUAFModuleValid(AnimNextComponent))
+				{
+					const UMovieSceneAnimMixerSettings* AnimMixerSettings = GetDefault<UMovieSceneAnimMixerSettings>();
+
+					FAnimMixerUtils::UnregisterWithUAFSubsystem(AnimNextComponent);
+					FAnimMixerUtils::SetUAFComponentModule(AnimNextComponent, AnimMixerSettings->DefaultUAFModule.LoadSynchronous());
+					FAnimMixerUtils::RegisterWithUAFSubsystem(AnimNextComponent);
+
+					bModuleReplaced = true;
+				}
 				UE::UAF::FInjectionSite InjectionSite(InjectionTarget->InjectionSite);
 				if (InjectionSite.DesiredSite.IsNone())
 				{
@@ -272,12 +302,14 @@ namespace UE::MovieScene
 
 					bool bAnimNextModuleEnabled = AnimNextComponent->IsEnabled();
 
-					auto OnCacheAnimatedState = [this, TargetData, bAnimNextModuleEnabled](UAnimNextComponent*)
+					auto OnCacheAnimatedState = [this, TargetData, bAnimNextModuleEnabled, bModuleReplaced, PreAnimatedModule](UAnimNextComponent*)
 					{
 						FPreAnimatedAnimNextState State;
 						State.InjectionHandle = TargetData->InjectionRequestHandle;
 						State.System = this->System;
 						State.bAnimNextModuleEnabled = bAnimNextModuleEnabled;
+						State.bModuleReplaced = bModuleReplaced;
+						State.PreAnimatedModule = PreAnimatedModule;
 						return State;
 					};
 
diff --git a/Engine/Plugins/Experimental/MovieSceneAnimMixer/Source/MovieSceneAnimMixer/Public/AnimMixerUtils.h b/Engine/Plugins/Experimental/MovieSceneAnimMixer/Source/MovieSceneAnimMixer/Public/AnimMixerUtils.h
new file mode 100644
--- /dev/null
+++ b/Engine/Plugins/Experimental/MovieSceneAnimMixer/Source/MovieSceneAnimMixer/Public/AnimMixerUtils.h
@@ -0,0 +1,43 @@
+﻿// Copyright Epic Games, Inc. All Rights Reserved.
+
+#pragma once
+#include "Component/AnimNextComponent.h"
+
+struct FAnimMixerUtils
+{
+
+	static void RegisterWithUAFSubsystem(UAnimNextComponent* InComponent)
+	{
+		if (InComponent)
+		{
+			InComponent->RegisterWithSubsystem();
+		}
+	}
+
+	static void UnregisterWithUAFSubsystem(UAnimNextComponent* InComponent)
+	{
+		if (InComponent)
+		{
+			InComponent->UnregisterWithSubsystem();
+		}
+	}
+
+	static void SetUAFComponentModule(UAnimNextComponent* InComponent, TObjectPtr<UAnimNextModule> InModule)
+	{
+		if (InComponent)
+		{
+			InComponent->SetModule(InModule);
+		}
+	}
+
+	static bool IsUAFModuleValid(UAnimNextComponent* InComponent)
+	{
+		return InComponent ? InComponent->IsModuleValid() : false;
+	}
+
+	static TObjectPtr<UAnimNextModule> GetUAFModule(UAnimNextComponent* InComponent)
+	{
+		return InComponent ? InComponent->Module : nullptr;
+	}
+	
+};
diff --git a/Engine/Plugins/Experimental/MovieSceneAnimMixer/Source/MovieSceneAnimMixer/Public/MovieSceneAnimMixerModule.h b/Engine/Plugins/Experimental/MovieSceneAnimMixer/Source/MovieSceneAnimMixer/Public/MovieSceneAnimMixerModule.h
--- a/Engine/Plugins/Experimental/MovieSceneAnimMixer/Source/MovieSceneAnimMixer/Public/MovieSceneAnimMixerModule.h
+++ b/Engine/Plugins/Experimental/MovieSceneAnimMixer/Source/MovieSceneAnimMixer/Public/MovieSceneAnimMixerModule.h
@@ -4,6 +4,7 @@
 
 #include "CoreMinimal.h"
 #include "Modules/ModuleInterface.h"
+#include "HAL/IConsoleManager.h"
 
 MOVIESCENEANIMMIXER_API DECLARE_LOG_CATEGORY_EXTERN(LogMovieSceneAnimMixer, Log, All);
 
@@ -18,13 +19,15 @@ DECLARE_MULTICAST_DELEGATE_TwoParams(FPreEvaluateMixerTasks, const FAnimationUpd
 
 namespace UE::MovieScene
 {
-
+	extern MOVIESCENEANIMMIXER_API FName GSequencerDefaultAnimNextInjectionSite;
+	
+	extern MOVIESCENEANIMMIXER_API FAutoConsoleVariableRef CVarDefaultAnimNextInjectionSite;
 
 class FMovieSceneAnimMixerModule : public IModuleInterface
 {
 	// IModuleInterface interface
-	virtual void StartupModule() override;
-	virtual void ShutdownModule() override;
+	MOVIESCENEANIMMIXER_API virtual void StartupModule() override;
+	MOVIESCENEANIMMIXER_API virtual void ShutdownModule() override;
 };
 
 }
diff --git a/Engine/Plugins/Experimental/MovieSceneAnimMixer/Source/MovieSceneAnimMixer/Public/MovieSceneAnimMixerSettings.h b/Engine/Plugins/Experimental/MovieSceneAnimMixer/Source/MovieSceneAnimMixer/Public/MovieSceneAnimMixerSettings.h
new file mode 100644
--- /dev/null
+++ b/Engine/Plugins/Experimental/MovieSceneAnimMixer/Source/MovieSceneAnimMixer/Public/MovieSceneAnimMixerSettings.h
@@ -0,0 +1,35 @@
+﻿// Copyright Epic Games, Inc. All Rights Reserved.
+#pragma once
+
+#include "CoreMinimal.h"
+#include "Engine/DeveloperSettings.h"
+
+#include "MovieSceneAnimMixerSettings.generated.h"
+
+class UAnimNextModule;
+
+UCLASS(config = Engine, defaultconfig, meta=(DisplayName="Anim Mixer"), MinimalAPI)
+class UMovieSceneAnimMixerSettings : public UDeveloperSettings
+{
+	GENERATED_BODY()
+
+public:
+	// Default injection site for the anim mixer to use, when no site is specified in their target. Can be overridden with the console variable Sequencer.AnimNext.DefaultInjectionSite
+	UPROPERTY(Config, EditAnywhere, Category="Animation", meta = (AllowedType = "FAnimNextAnimGraph"))
+	FName DefaultInjectionSite;
+
+	// The default UAF module to use when a mixer is animating a mesh with an AnimNext component
+	UPROPERTY(Config, EditAnywhere, Category="Module")
+	TSoftObjectPtr<UAnimNextModule> DefaultUAFModule = nullptr;
+
+	MOVIESCENEANIMMIXER_API virtual FName GetCategoryName() const override;
+
+	
+
+#if WITH_EDITOR
+	virtual void PostEditChangeProperty(FPropertyChangedEvent& PropertyChangedEvent) override;
+	
+	MOVIESCENEANIMMIXER_API virtual FText GetSectionText() const override;
+#endif
+	
+};
diff --git a/Engine/Plugins/Experimental/UAF/UAF/Source/UAF/Private/Component/AnimNextComponent.cpp b/Engine/Plugins/Experimental/UAF/UAF/Source/UAF/Private/Component/AnimNextComponent.cpp
--- a/Engine/Plugins/Experimental/UAF/UAF/Source/UAF/Private/Component/AnimNextComponent.cpp
+++ b/Engine/Plugins/Experimental/UAF/UAF/Source/UAF/Private/Component/AnimNextComponent.cpp
@@ -469,3 +469,24 @@ void UAnimNextComponent::SetModule(TObjectPtr<UAnimNextModule> InModule)
 {
 	Module = InModule;
 }
+
+void UAnimNextComponent::RegisterWithSubsystem()
+{
+	if (Subsystem && Module)
+	{
+		Subsystem->Register(this);
+	}
+}
+
+void UAnimNextComponent::UnregisterWithSubsystem()
+{
+	if (Subsystem)
+	{
+		Subsystem->Unregister(this);
+	}
+}
+
+bool UAnimNextComponent::IsModuleValid()
+{
+	return ModuleHandle.IsValid();
+}
diff --git a/Engine/Plugins/Experimental/UAF/UAF/Source/UAF/Public/Component/AnimNextComponent.h b/Engine/Plugins/Experimental/UAF/UAF/Source/UAF/Public/Component/AnimNextComponent.h
--- a/Engine/Plugins/Experimental/UAF/UAF/Source/UAF/Public/Component/AnimNextComponent.h
+++ b/Engine/Plugins/Experimental/UAF/UAF/Source/UAF/Public/Component/AnimNextComponent.h
@@ -16,6 +16,8 @@ struct FAnimNextComponentInstanceData;
 class UAnimNextComponentWorldSubsystem;
 class UAnimNextModule;
 struct AnimNextVariableReference;
+// TODO: if a public API is implemented, the forward declaration of this struct would no longer be needed.
+struct FAnimMixerUtils;
 
 namespace UE::UAF::UncookedOnly
 {
@@ -187,6 +189,11 @@ public:
 private:
 	UAF_API void SetModule(TObjectPtr<UAnimNextModule> InModule);
 
+	UAF_API void RegisterWithSubsystem();
+	UAF_API void UnregisterWithSubsystem();
+
+	UAF_API bool IsModuleValid();
+
 	DECLARE_FUNCTION(execBlueprintSetVariable);
 	DECLARE_FUNCTION(execBlueprintSetVariableReference);
 
@@ -201,6 +208,8 @@ private:
 	friend struct FAssetBlueprintTests;
 	friend class UAnimNextComponentWorldSubsystem;
 	friend class UE::UAF::Editor::FAssetWizard;
+	// TODO ideally there would be a public API for the things FAnimMixerUtils does, so it wouldn't be necessary.
+	friend struct FAnimMixerUtils;
 
 	// The UAF system that this component will run
 	UPROPERTY(EditAnywhere, Category="System", DisplayName="System")
