Sequencer: Added an option to prevent Anim Mixer Tracks from overriding Animation Track behavior

By disabling Sequencer.AnimMixer.OverrideAnimationTrack, the animation mier will leave regular Animation Tracks alone, causing them to use their previous evaluation logic. This is intended as a temporary fallback while stabilization work on the Anim Mixer is finished, but causes Animation Tracks and Anim Mixer Tracks on the _same object_ to be incompatible: in such situations the tracks will fight with each other rather than blend naturally.

diff --git a/Engine/Plugins/Experimental/MovieSceneAnimMixer/Source/MovieSceneAnimMixer/Internal/MovieSceneAnimationMixerTrack.h b/Engine/Plugins/Experimental/MovieSceneAnimMixer/Source/MovieSceneAnimMixer/Internal/MovieSceneAnimationMixerTrack.h
--- a/Engine/Plugins/Experimental/MovieSceneAnimMixer/Source/MovieSceneAnimMixer/Internal/MovieSceneAnimationMixerTrack.h
+++ b/Engine/Plugins/Experimental/MovieSceneAnimMixer/Source/MovieSceneAnimMixer/Internal/MovieSceneAnimationMixerTrack.h
@@ -79,11 +79,14 @@ public:
 UCLASS(MinimalAPI)
 class UMovieSceneAnimationSectionDecoration
 	: public UMovieSceneSignedObject
+	, public IMovieSceneEntityDecorator
 	, public IMovieSceneSectionDecoration
 	, public IMovieSceneAnimationSectionInterface
 {
 	GENERATED_BODY()
 
+	virtual void ExtendEntityImpl(UMovieSceneEntitySystemLinker* EntityLinker, const FEntityImportParams& Params, FImportedEntity* OutImportedEntity) override;
+
 	virtual int32 GetRowSortOrder() const
 	{
 		return RowSortOrder;
diff --git a/Engine/Plugins/Experimental/MovieSceneAnimMixer/Source/MovieSceneAnimMixer/Private/MovieSceneAnimationMixerTrack.cpp b/Engine/Plugins/Experimental/MovieSceneAnimMixer/Source/MovieSceneAnimMixer/Private/MovieSceneAnimationMixerTrack.cpp
--- a/Engine/Plugins/Experimental/MovieSceneAnimMixer/Source/MovieSceneAnimMixer/Private/MovieSceneAnimationMixerTrack.cpp
+++ b/Engine/Plugins/Experimental/MovieSceneAnimMixer/Source/MovieSceneAnimMixer/Private/MovieSceneAnimationMixerTrack.cpp
@@ -5,6 +5,7 @@
 #include "Sections/MovieSceneSkeletalAnimationSection.h"
 #include "Channels/MovieSceneChannelProxy.h"
 #include "EntitySystem/BuiltInComponentTypes.h"
+#include "MovieSceneTracksComponentTypes.h"
 #include "AnimMixerComponentTypes.h"
 #include "Misc/AxisDisplayInfo.h"
 #include "Systems/MovieSceneRootMotionSystem.h"
@@ -180,6 +181,15 @@ void UMovieSceneAnimationBaseTransformDecoration::ExtendEntityImpl(UMovieSceneEn
 	);
 }
 
+void UMovieSceneAnimationSectionDecoration::ExtendEntityImpl(UMovieSceneEntitySystemLinker* EntityLinker, const FEntityImportParams& Params, FImportedEntity* OutImportedEntity)
+{
+	using namespace UE::MovieScene;
+
+	OutImportedEntity->AddBuilder(
+		FEntityBuilder()
+		.AddTag(FMovieSceneTracksComponentTypes::Get()->Tags.AnimMixerPoseProducer)
+	);
+}
 
 UMovieSceneAnimationMixerTrack::UMovieSceneAnimationMixerTrack(const FObjectInitializer& ObjInit)
 	: Super(ObjInit)
diff --git a/Engine/Plugins/Experimental/MovieSceneAnimMixer/Source/MovieSceneAnimMixer/Private/Systems/MovieSceneMixedSkeletalAnimationSystem.cpp b/Engine/Plugins/Experimental/MovieSceneAnimMixer/Source/MovieSceneAnimMixer/Private/Systems/MovieSceneMixedSkeletalAnimationSystem.cpp
--- a/Engine/Plugins/Experimental/MovieSceneAnimMixer/Source/MovieSceneAnimMixer/Private/Systems/MovieSceneMixedSkeletalAnimationSystem.cpp
+++ b/Engine/Plugins/Experimental/MovieSceneAnimMixer/Source/MovieSceneAnimMixer/Private/Systems/MovieSceneMixedSkeletalAnimationSystem.cpp
@@ -49,13 +49,20 @@
 namespace UE::MovieScene
 {
 
-float GMovieSceneAnimMixerEnabled = true;
+bool GMovieSceneAnimMixerEnabled = true;
 FAutoConsoleVariableRef CVarMovieSceneAnimMixerEnabled(
 	TEXT("Sequencer.AnimMixer.Enabled"),
 	GMovieSceneAnimMixerEnabled,
 	TEXT("(Default: true) Controls whether the new Anim Mixer library is used for skeletal animation evaluation.\n"),
 	ECVF_Default);
 
+bool GAnimMixerOverrideAnimationTrack = true;
+FAutoConsoleVariableRef CVarAnimMixerOverrideAnimationTrack(
+	TEXT("Sequencer.AnimMixer.OverrideAnimationTrack"),
+	GAnimMixerOverrideAnimationTrack,
+	TEXT("(Default: true) Controls whether the new Anim Mixer library should override Animation Track behavior or not. Warning: disabling this means that Animation Tracks and Animation Mixer Tracks will fight with each other if they animate the same object.\n"),
+	ECVF_Default);
+
 // Adds required anim mixer components onto skeletal animation section entities
 struct FSkeletalAnimMixerMutation : IMovieSceneEntityMutation
 {
@@ -549,6 +556,13 @@ void UMovieSceneMixedSkeletalAnimationSystem::OnRun(FSystemTaskPrerequisites& In
 		Filter.All({ BuiltInComponents->BoundObject, TrackComponents->SkeletalAnimation, BuiltInComponents->Tags.NeedsLink });
 		Filter.None({ BuiltInComponents->Tags.ImportedEntity });
 
+		if (GAnimMixerOverrideAnimationTrack == false)
+		{
+			// If we're not overriding animation tracks, constrain the mutation to things that have already been marked as AnimMixerPoseProducers
+			//   This happens for the Anim Mixer Track inside UMovieSceneAnimationSectionDecoration::ExtendEntityImpl
+			Filter.All({ TrackComponents->Tags.AnimMixerPoseProducer });
+		}
+
 		// Initialize components for skeletal animation mixing
 		FSkeletalAnimMixerMutation SkeletalAnimMixerMutation(Linker);
 		Linker->EntityManager.MutateAll(Filter, SkeletalAnimMixerMutation);
