Sequencer: Partial back out of 42444020 and 42182253 which added UnbindFromSkeletalMeshComponent in response to SkeletalMeshUnregistered

This caused UE-282772 because unbinding Sequencer when editing in the details panel and then forcing a rebind during editing in Sequencer resulted in an interruption in the undo stack

diff --git a/Engine/Plugins/Animation/ControlRig/Source/ControlRig/Private/Sequencer/MovieSceneControlRigParameterTemplate.cpp b/Engine/Plugins/Animation/ControlRig/Source/ControlRig/Private/Sequencer/MovieSceneControlRigParameterTemplate.cpp
--- a/Engine/Plugins/Animation/ControlRig/Source/ControlRig/Private/Sequencer/MovieSceneControlRigParameterTemplate.cpp
+++ b/Engine/Plugins/Animation/ControlRig/Source/ControlRig/Private/Sequencer/MovieSceneControlRigParameterTemplate.cpp
@@ -1347,19 +1347,7 @@ struct FControlRigParameterExecutionToken : IMovieSceneExecutionToken
 								}
 							});
 							
-							auto UnregisteredLambda = USkeletalMeshComponent::FOnSkeletalMeshUnregisteredDelegate::CreateLambda([WeakSkeletalMeshComponentPtr](const USkeletalMeshComponent* InMeshComponent)
-							{
-								const TStrongObjectPtr<USkeletalMeshComponent> PinnedSkeletalMeshComponent = WeakSkeletalMeshComponentPtr.Pin();
-								if (PinnedSkeletalMeshComponent && InMeshComponent == PinnedSkeletalMeshComponent.Get())
-								{
-									FAnimCustomInstanceHelper::UnbindFromSkeletalMeshComponent<UControlRigLayerInstance>(PinnedSkeletalMeshComponent.Get());
-								}
-							});
-
-							FDelegateHandle SkeletalMeshUnregisteredHandle = SkeletalMeshComponent->RegisterOnSkeletalMeshUnregistered(UnregisteredLambda);
-
 							Template->PreCompileHandles.Add(PreCompileHandle);
-							Template->MeshUnregisteredHandles.Add({SkeletalMeshUnregisteredHandle, SkeletalMeshComponent});
 						}
 						
 #endif
@@ -2519,14 +2507,6 @@ FMovieSceneControlRigParameterTemplate::~FMovieSceneControlRigParameterTemplate(
 		{
 			GEditor->OnBlueprintPreCompile().Remove(Handle);
 		}
-
-		for (TTuple<FDelegateHandle, TWeakObjectPtr<USkeletalMeshComponent>>& HandlePair : MeshUnregisteredHandles)
-		{
-			if (TStrongObjectPtr<USkeletalMeshComponent> SkeletalMeshComponent = HandlePair.Value.Pin())
-			{
-				SkeletalMeshComponent->UnregisterOnSkeletalMeshUnregistered(HandlePair.Key);
-			}
-		}
 	}
 }
 #endif
diff --git a/Engine/Plugins/Animation/ControlRig/Source/ControlRig/Private/Sequencer/MovieSceneControlRigSystem.cpp b/Engine/Plugins/Animation/ControlRig/Source/ControlRig/Private/Sequencer/MovieSceneControlRigSystem.cpp
--- a/Engine/Plugins/Animation/ControlRig/Source/ControlRig/Private/Sequencer/MovieSceneControlRigSystem.cpp
+++ b/Engine/Plugins/Animation/ControlRig/Source/ControlRig/Private/Sequencer/MovieSceneControlRigSystem.cpp
@@ -2011,19 +2011,7 @@ void UMovieSceneControlRigParameterEvaluatorSystem::InitializeBaseRigComponent(
 					}
 				});
 								
-				auto UnregisteredLambda = USkeletalMeshComponent::FOnSkeletalMeshUnregisteredDelegate::CreateLambda([WeakSkeletalMeshComponentPtr](const USkeletalMeshComponent* InMeshComponent)
-				{
-					const TStrongObjectPtr<USkeletalMeshComponent> PinnedSkeletalMeshComponent = WeakSkeletalMeshComponentPtr.Pin();
-					if (PinnedSkeletalMeshComponent && InMeshComponent == PinnedSkeletalMeshComponent.Get())
-					{
-						FAnimCustomInstanceHelper::UnbindFromSkeletalMeshComponent<UControlRigLayerInstance>(PinnedSkeletalMeshComponent.Get());
-					}
-				});
-
-				FDelegateHandle SkeletalMeshUnregisteredHandle = SkeletalMeshComponent->RegisterOnSkeletalMeshUnregistered(UnregisteredLambda);
-
 				PreCompileHandles.Add(PreCompileHandle);
-				MeshUnregisteredHandles.Add({SkeletalMeshUnregisteredHandle, SkeletalMeshComponent});
 			}
 		}
 	}
@@ -2040,14 +2028,6 @@ UMovieSceneControlRigParameterEvaluatorSystem::~UMovieSceneControlRigParameterEv
 		{
 			GEditor->OnBlueprintPreCompile().Remove(Handle);
 		}
-
-		for (TTuple<FDelegateHandle, TWeakObjectPtr<USkeletalMeshComponent>>& HandlePair : MeshUnregisteredHandles)
-		{
-			if (TStrongObjectPtr<USkeletalMeshComponent> SkeletalMeshComponent = HandlePair.Value.Pin())
-			{
-				SkeletalMeshComponent->UnregisterOnSkeletalMeshUnregistered(HandlePair.Key);
-			}
-		}
 	}
 }
 #endif
diff --git a/Engine/Plugins/Animation/ControlRig/Source/ControlRig/Public/Sequencer/MovieSceneControlRigParameterTemplate.h b/Engine/Plugins/Animation/ControlRig/Source/ControlRig/Public/Sequencer/MovieSceneControlRigParameterTemplate.h
--- a/Engine/Plugins/Animation/ControlRig/Source/ControlRig/Public/Sequencer/MovieSceneControlRigParameterTemplate.h
+++ b/Engine/Plugins/Animation/ControlRig/Source/ControlRig/Public/Sequencer/MovieSceneControlRigParameterTemplate.h
@@ -24,7 +24,6 @@ struct FMovieSceneControlRigParameterTemplate : public FMovieSceneParameterSecti
 	virtual ~FMovieSceneControlRigParameterTemplate() override;
 
 	TArray<FDelegateHandle> PreCompileHandles;
-	TArray<TTuple<FDelegateHandle, TWeakObjectPtr<USkeletalMeshComponent>>> MeshUnregisteredHandles;
 #endif
 	
 private:
diff --git a/Engine/Plugins/Animation/ControlRig/Source/ControlRig/Public/Sequencer/MovieSceneControlRigSystem.h b/Engine/Plugins/Animation/ControlRig/Source/ControlRig/Public/Sequencer/MovieSceneControlRigSystem.h
--- a/Engine/Plugins/Animation/ControlRig/Source/ControlRig/Public/Sequencer/MovieSceneControlRigSystem.h
+++ b/Engine/Plugins/Animation/ControlRig/Source/ControlRig/Public/Sequencer/MovieSceneControlRigSystem.h
@@ -244,7 +244,6 @@ public:
 	virtual  ~UMovieSceneControlRigParameterEvaluatorSystem() override;
 	
 	TArray<FDelegateHandle> PreCompileHandles;
-	TArray<TTuple<FDelegateHandle, TWeakObjectPtr<USkeletalMeshComponent>>> MeshUnregisteredHandles;
 #endif
 	
 };
diff --git a/Engine/Plugins/Experimental/MovieSceneAnimMixer/Source/MovieSceneAnimMixer/Private/Systems/MovieSceneAnimInstanceTargetSystem.cpp b/Engine/Plugins/Experimental/MovieSceneAnimMixer/Source/MovieSceneAnimMixer/Private/Systems/MovieSceneAnimInstanceTargetSystem.cpp
--- a/Engine/Plugins/Experimental/MovieSceneAnimMixer/Source/MovieSceneAnimMixer/Private/Systems/MovieSceneAnimInstanceTargetSystem.cpp
+++ b/Engine/Plugins/Experimental/MovieSceneAnimMixer/Source/MovieSceneAnimMixer/Private/Systems/MovieSceneAnimInstanceTargetSystem.cpp
@@ -482,19 +482,8 @@ namespace UE::MovieScene
 						FAnimCustomInstanceHelper::UnbindFromSkeletalMeshComponent<USequencerMixedAnimInstance>(SkeletalMeshComponentWeakPtr.Get());
 					}
 				});
-				
-				auto UnregisteredLambda = USkeletalMeshComponent::FOnSkeletalMeshUnregisteredDelegate::CreateLambda([CurrentAnimInstanceWeakPtr, SkeletalMeshComponentWeakPtr](const USkeletalMeshComponent* InMeshComponent)
-				{
-					if (CurrentAnimInstanceWeakPtr.IsValid() && SkeletalMeshComponentWeakPtr.IsValid() && InMeshComponent == SkeletalMeshComponentWeakPtr.Get())
-					{
-						FAnimCustomInstanceHelper::UnbindFromSkeletalMeshComponent<USequencerMixedAnimInstance>(SkeletalMeshComponentWeakPtr.Get());
-					}
-				});
-
-				FDelegateHandle SkeletalMeshUnregisteredHandle = SkeletalMeshComponent->RegisterOnSkeletalMeshUnregistered(UnregisteredLambda);
 			
 				System->PreCompileHandles.Add(PreCompileHandle);
-				System->MeshUnregisteredHandles.Add({SkeletalMeshUnregisteredHandle, SkeletalMeshComponentWeakPtr});
 			}
 #endif
 			
@@ -558,14 +547,6 @@ UMovieSceneAnimInstanceTargetSystem::~UMovieSceneAnimInstanceTargetSystem()
 		{
 			GEditor->OnBlueprintPreCompile().Remove(Handle);
 		}
-
-		for (TTuple<FDelegateHandle, TWeakObjectPtr<USkeletalMeshComponent>> HandlePair : MeshUnregisteredHandles)
-		{
-			if (TStrongObjectPtr<USkeletalMeshComponent> SkeletalMeshComponent = HandlePair.Value.Pin())
-			{
-				SkeletalMeshComponent->UnregisterOnSkeletalMeshUnregistered(HandlePair.Key);
-			}
-		}
 	}
 }
 #endif
diff --git a/Engine/Plugins/Experimental/MovieSceneAnimMixer/Source/MovieSceneAnimMixer/Public/Systems/MovieSceneAnimInstanceTargetSystem.h b/Engine/Plugins/Experimental/MovieSceneAnimMixer/Source/MovieSceneAnimMixer/Public/Systems/MovieSceneAnimInstanceTargetSystem.h
--- a/Engine/Plugins/Experimental/MovieSceneAnimMixer/Source/MovieSceneAnimMixer/Public/Systems/MovieSceneAnimInstanceTargetSystem.h
+++ b/Engine/Plugins/Experimental/MovieSceneAnimMixer/Source/MovieSceneAnimMixer/Public/Systems/MovieSceneAnimInstanceTargetSystem.h
@@ -124,7 +124,6 @@ public:
 	virtual ~UMovieSceneAnimInstanceTargetSystem() override;
 
 	TArray<FDelegateHandle> PreCompileHandles;
-	TArray<TTuple<FDelegateHandle, TWeakObjectPtr<USkeletalMeshComponent>>> MeshUnregisteredHandles;
 #endif
 	
 private:
