Sequencer Anim Mixer: Skel Anim sections and Stitch sections not under a mixer track will no longer have a mixer target or priority property for the user to set. They will still go through the mixer but with default settings. If a user wants to use a non-default target, they will need to create an anim mixer and set the target on that.

The Anim Mixer track now correctly sets the target component and the priority component on sections underneath it. It does that by adding the UMovieSceneAnimationSectionDecoration to these sections, which extends their entity with these components.

As part of this change, fixed some backwards compatibility issues with the skeletal anim track when not using a mixer track, specifically with slot nodes in anim bp. Now slot nodes will check for anim mixer presence and allow the sequence anim bp target to amend their slot pose the same as if they were using the new sequencer anim mixer target node. Also, 'Force Custom Mode' is now respected.

diff --git a/Engine/Plugins/Experimental/MovieSceneAnimMixer/Source/MovieSceneAnimMixer/Private/AnimNode_SequencerMixerTarget.cpp b/Engine/Plugins/Experimental/MovieSceneAnimMixer/Source/MovieSceneAnimMixer/Private/AnimNode_SequencerMixerTarget.cpp
--- a/Engine/Plugins/Experimental/MovieSceneAnimMixer/Source/MovieSceneAnimMixer/Private/AnimNode_SequencerMixerTarget.cpp
+++ b/Engine/Plugins/Experimental/MovieSceneAnimMixer/Source/MovieSceneAnimMixer/Private/AnimNode_SequencerMixerTarget.cpp
@@ -40,11 +40,16 @@ void FAnimNode_SequencerMixerTarget::Evaluate_AnyThread(FPoseContext& Output)
 {
 	SourcePose.Evaluate(Output);
 
+	ApplySequencerMixedPose(Output, TargetName);
+}
+
+void FAnimNode_SequencerMixerTarget::ApplySequencerMixedPose(FPoseContext& Output, FName InTargetName)
+{
 	UAnimInstance* AnimInstance = Cast<UAnimInstance>(Output.GetAnimInstanceObject());
 	
 	const FAnimSubsystem_SequencerMixer& MixerSubsystem = AnimInstance->GetSubsystem<FAnimSubsystem_SequencerMixer>();
 	
-	const TSharedPtr<FAnimNextEvaluationTask>* MixerTask = MixerSubsystem.GetEvalTask(TargetName);
+	const TSharedPtr<FAnimNextEvaluationTask>* MixerTask = MixerSubsystem.GetEvalTask(InTargetName);
 	if (!MixerTask || !MixerTask->IsValid())
 	{
 		return;
@@ -105,5 +110,4 @@ void FAnimNode_SequencerMixerTarget::Evaluate_AnyThread(FPoseContext& Output)
 	FGenerationTools::RemapPose(ResultPose.LODPose, Output);
 	Output.Curve.CopyFrom(ResultPose.Curves);
 	FGenerationTools::RemapAttributes(ResultPose.LODPose, ResultPose.Attributes, Output);
-
 }
diff --git a/Engine/Plugins/Experimental/MovieSceneAnimMixer/Source/MovieSceneAnimMixer/Private/MovieSceneAnimMixerModule.cpp b/Engine/Plugins/Experimental/MovieSceneAnimMixer/Source/MovieSceneAnimMixer/Private/MovieSceneAnimMixerModule.cpp
--- a/Engine/Plugins/Experimental/MovieSceneAnimMixer/Source/MovieSceneAnimMixer/Private/MovieSceneAnimMixerModule.cpp
+++ b/Engine/Plugins/Experimental/MovieSceneAnimMixer/Source/MovieSceneAnimMixer/Private/MovieSceneAnimMixerModule.cpp
@@ -1,6 +1,8 @@
 // Copyright Epic Games, Inc. All Rights Reserved.
 
 #include "MovieSceneAnimMixerModule.h"
+
+#include "AnimNode_SequencerMixerTarget.h"
 #include "MovieSceneAnimMixerSettings.h"
 #include "ISettingsModule.h"
 #include "Modules/ModuleInterface.h"
@@ -33,11 +35,17 @@ void FMovieSceneAnimMixerModule::StartupModule()
 
 	CVarDefaultAnimNextInjectionSite->Set(*GetMutableDefault<UMovieSceneAnimMixerSettings>()->DefaultInjectionSite.ToString(), ECVF_SetByProjectSetting);
 	
+	IModularFeatures::Get().RegisterModularFeature(IAnimGraphRuntime_SequencerMixerTargetConnector::GetModularFeatureName(), this);
 }
 
 void FMovieSceneAnimMixerModule::ShutdownModule()
 {
+	IModularFeatures::Get().UnregisterModularFeature(IAnimGraphRuntime_SequencerMixerTargetConnector::GetModularFeatureName(), this);
+}
 
+void FMovieSceneAnimMixerModule::ApplySequencerMixedPose(FPoseContext& Output, FName InTargetName) const
+{
+	FAnimNode_SequencerMixerTarget::ApplySequencerMixedPose(Output, InTargetName);
 }
 
 }
diff --git a/Engine/Plugins/Experimental/MovieSceneAnimMixer/Source/MovieSceneAnimMixer/Private/MovieSceneAnimationMixerTrack.cpp b/Engine/Plugins/Experimental/MovieSceneAnimMixer/Source/MovieSceneAnimMixer/Private/MovieSceneAnimationMixerTrack.cpp
--- a/Engine/Plugins/Experimental/MovieSceneAnimMixer/Source/MovieSceneAnimMixer/Private/MovieSceneAnimationMixerTrack.cpp
+++ b/Engine/Plugins/Experimental/MovieSceneAnimMixer/Source/MovieSceneAnimMixer/Private/MovieSceneAnimationMixerTrack.cpp
@@ -184,10 +184,28 @@ void UMovieSceneAnimationBaseTransformDecoration::ExtendEntityImpl(UMovieSceneEn
 void UMovieSceneAnimationSectionDecoration::ExtendEntityImpl(UMovieSceneEntitySystemLinker* EntityLinker, const FEntityImportParams& Params, FImportedEntity* OutImportedEntity)
 {
 	using namespace UE::MovieScene;
+	FAnimMixerComponentTypes* AnimMixerComponents = FAnimMixerComponentTypes::Get();
+
+	UMovieSceneAnimationMixerTrack* AnimationMixerTrack = GetTypedOuter<UMovieSceneAnimationMixerTrack>();
+
+	TInstancedStruct<FMovieSceneMixedAnimationTarget> Target;
+	if (AnimationMixerTrack != nullptr && AnimationMixerTrack->MixedAnimationTarget.IsValid())
+	{
+		Target = AnimationMixerTrack->MixedAnimationTarget;
+	}
+	else
+	{
+		Target = TInstancedStruct<FMovieSceneMixedAnimationTarget>::Make();
+	}
+
+	UMovieSceneSection* SectionOuter = GetTypedOuter<UMovieSceneSection>();
 
 	OutImportedEntity->AddBuilder(
 		FEntityBuilder()
 		.AddTag(FMovieSceneTracksComponentTypes::Get()->Tags.AnimMixerPoseProducer)
+		.Add(AnimMixerComponents->Target, Target)		
+		.Add(AnimMixerComponents->Priority, SectionOuter ? SectionOuter->GetRowIndex() : GetRowSortOrder())
+		.AddMutualComponents()
 	);
 }
 
diff --git a/Engine/Plugins/Experimental/MovieSceneAnimMixer/Source/MovieSceneAnimMixer/Private/Systems/MovieSceneAnimMixerSystem.cpp b/Engine/Plugins/Experimental/MovieSceneAnimMixer/Source/MovieSceneAnimMixer/Private/Systems/MovieSceneAnimMixerSystem.cpp
--- a/Engine/Plugins/Experimental/MovieSceneAnimMixer/Source/MovieSceneAnimMixer/Private/Systems/MovieSceneAnimMixerSystem.cpp
+++ b/Engine/Plugins/Experimental/MovieSceneAnimMixer/Source/MovieSceneAnimMixer/Private/Systems/MovieSceneAnimMixerSystem.cpp
@@ -44,6 +44,7 @@
 #include "Systems/MovieSceneComponentTransformSystem.h"
 #include "Systems/MovieSceneRootMotionSystem.h"
 #include "AnimSubsystem_SequencerMixer.h"
+#include "Sections/MovieSceneSkeletalAnimationSection.h"
 
 #include UE_INLINE_GENERATED_CPP_BY_NAME(MovieSceneAnimMixerSystem)
 
@@ -683,6 +684,7 @@ void UMovieSceneAnimMixerSystem::OnRun(FSystemTaskPrerequisites& InPrerequisites
 		.Write(AnimMixerComponents->MixerEntry)
 		.Read(AnimMixerComponents->Task)
 		.Read(AnimMixerComponents->Priority)
+		.ReadOptional(TrackComponents->SkeletalAnimation)
 		.ReadOptional(AnimMixerComponents->RootMotionSettings)
 		.FilterAll({ BuiltInComponents->Tags.NeedsLink})
 		.Iterate_PerEntity( &Linker->EntityManager,
@@ -695,10 +697,11 @@ void UMovieSceneAnimMixerSystem::OnRun(FSystemTaskPrerequisites& InPrerequisites
 			TSharedPtr<FMovieSceneAnimMixerEntry>& InOutMixerEntry,
 			TSharedPtr<FAnimNextEvaluationTask> Task,
 			int32 Priority,
+			const FMovieSceneSkeletalAnimationComponentData* SkelAnimData,
 			const FMovieSceneRootMotionSettings* RootMotionSettings)
 			{
 				// For new entities, we 'resolve' the animation target so if 'Automatic' is picked we choose the right one automatically.
-				Target = ResolveAnimationTarget(BoundObjectKey, Target);
+				Target = ResolveAnimationTarget(BoundObjectKey, Target, SkelAnimData);
 				FMovieSceneAnimMixerKey Key({ BoundObjectKey, Target });
 				TSharedPtr<FMovieSceneAnimMixer> Mixer = Mixers.FindRef(Key);
 				if (!Mixer)
@@ -998,7 +1001,7 @@ void UMovieSceneAnimMixerSystem::OnCleanTaggedGarbage()
 	);
 }
 
-TInstancedStruct<FMovieSceneMixedAnimationTarget> UMovieSceneAnimMixerSystem::ResolveAnimationTarget(FObjectKey ObjectKey, const TInstancedStruct<FMovieSceneMixedAnimationTarget>& InTarget)
+TInstancedStruct<FMovieSceneMixedAnimationTarget> UMovieSceneAnimMixerSystem::ResolveAnimationTarget(FObjectKey ObjectKey, const TInstancedStruct<FMovieSceneMixedAnimationTarget>& InTarget, const FMovieSceneSkeletalAnimationComponentData* SkelAnimData)
 {
 	// If user has selected the default 'automatic' target, attempt to choose one automatically for them.
 	if (!InTarget.IsValid() || InTarget.GetScriptStruct() == FMovieSceneMixedAnimationTarget::StaticStruct())
@@ -1027,12 +1030,26 @@ TInstancedStruct<FMovieSceneMixedAnimationTarget> UMovieSceneAnimMixerSystem::Re
 			}
 			else if (SkeletalMeshComponent)
 			{
-				if (UAnimInstance* AnimInstance = SkeletalMeshComponent->GetAnimInstance())
+				// Backwards compatibility- are they forcing custom mode?
+				const bool bForceCustomMode = SkelAnimData != nullptr && SkelAnimData->Section && SkelAnimData->Section->Params.bForceCustomMode;
+
+				if (!bForceCustomMode)
 				{
-					if (const FAnimSubsystem_SequencerMixer* MixerSubsystem = AnimInstance->FindSubsystem<FAnimSubsystem_SequencerMixer>())
+					if (UAnimInstance* AnimInstance = SkeletalMeshComponent->GetAnimInstance())
 					{
-						// We have an anim blueprint with sequencer mixer target node(s). Create a target using the default target name.
-						return TInstancedStruct<FMovieSceneAnimBlueprintTarget>::Make();
+						if (const FAnimSubsystem_SequencerMixer* MixerSubsystem = AnimInstance->FindSubsystem<FAnimSubsystem_SequencerMixer>())
+						{
+							FName BlueprintNodeName = FAnimNode_SequencerMixerTarget::DefaultTargetName;
+						
+							// Backwards compatibility- if we have skel anim data with a 'SlotName', use that as the name.
+							if (SkelAnimData != nullptr && SkelAnimData->Section && !SkelAnimData->Section->Params.SlotName.IsNone())
+							{
+								BlueprintNodeName = SkelAnimData->Section->Params.SlotName;
+							}
+						
+							// We have an anim blueprint with sequencer mixer target node(s). Create a target using the default target name.
+							return TInstancedStruct<FMovieSceneAnimBlueprintTarget>::Make((FMovieSceneAnimBlueprintTarget(BlueprintNodeName)));
+						}
 					}
 				}
 
diff --git a/Engine/Plugins/Experimental/MovieSceneAnimMixer/Source/MovieSceneAnimMixer/Private/Systems/MovieSceneMixedSkeletalAnimationSystem.cpp b/Engine/Plugins/Experimental/MovieSceneAnimMixer/Source/MovieSceneAnimMixer/Private/Systems/MovieSceneMixedSkeletalAnimationSystem.cpp
--- a/Engine/Plugins/Experimental/MovieSceneAnimMixer/Source/MovieSceneAnimMixer/Private/Systems/MovieSceneMixedSkeletalAnimationSystem.cpp
+++ b/Engine/Plugins/Experimental/MovieSceneAnimMixer/Source/MovieSceneAnimMixer/Private/Systems/MovieSceneMixedSkeletalAnimationSystem.cpp
@@ -127,19 +127,17 @@ struct FSkeletalAnimMixerMutation : IMovieSceneEntityMutation
 			const FMovieSceneSkeletalAnimationComponentData& SkeletalAnim = SkeletalAnimationData[Index];
 			if (SkeletalAnim.Section)
 			{
-				if (SkeletalAnim.Section->MixedAnimationTarget.IsValid())
+				UMovieSceneAnimationMixerTrack* MixerTrack = Cast<UMovieSceneAnimationMixerTrack>(SkeletalAnim.Section->GetOuter());
+				
+				// If we have a mixer track, the animation decoration may already have added our target and priority
+				
+				if (MixerTrack == nullptr || !MixerTrack->MixedAnimationTarget.IsValid())
 				{
-					// Set the assigned target
-					OutAnimTargets[Index] = SkeletalAnim.Section->MixedAnimationTarget;
-				}
-				else
-				{
-					// Set a default target
+					// Set a default target- Automatic will resolve to the appropriate place, or the Anim Mixer track will override it.
 					OutAnimTargets[Index] = TInstancedStruct<FMovieSceneMixedAnimationTarget>::Make();
+					OutPriorities[Index] = 0;
 				}
 
-				OutPriorities[Index] = SkeletalAnim.Section->MixedAnimationPriority;
-
 				UMovieSceneCommonAnimationTrack* Track = SkeletalAnim.Section->GetTypedOuter<UMovieSceneCommonAnimationTrack>();
 				UMovieSceneSkeletalAnimationTrack* SkelAnimTrack = Cast<UMovieSceneSkeletalAnimationTrack>(Track);
 
diff --git a/Engine/Plugins/Experimental/MovieSceneAnimMixer/Source/MovieSceneAnimMixer/Public/AnimNode_SequencerMixerTarget.h b/Engine/Plugins/Experimental/MovieSceneAnimMixer/Source/MovieSceneAnimMixer/Public/AnimNode_SequencerMixerTarget.h
--- a/Engine/Plugins/Experimental/MovieSceneAnimMixer/Source/MovieSceneAnimMixer/Public/AnimNode_SequencerMixerTarget.h
+++ b/Engine/Plugins/Experimental/MovieSceneAnimMixer/Source/MovieSceneAnimMixer/Public/AnimNode_SequencerMixerTarget.h
@@ -21,6 +21,9 @@ struct FAnimNode_SequencerMixerTarget final : public FAnimNode_Base
 	UPROPERTY(EditAnywhere, BlueprintReadWrite, Category=Setttings, meta=(CustomizePropery))
 	FName TargetName;
 
+	// We expose this so that the anim slot node can (through IAnimGraph_SequencerMixerTargetConnector) be backwards compatible with the mixer when the mixer plugin is enabled
+	static void ApplySequencerMixedPose(FPoseContext& Output, FName InTargetName);
+
 	// FAnimNode_Base overrides
 	MOVIESCENEANIMMIXER_API FAnimNode_SequencerMixerTarget();
 	MOVIESCENEANIMMIXER_API virtual void Initialize_AnyThread(const FAnimationInitializeContext& Context) override;
diff --git a/Engine/Plugins/Experimental/MovieSceneAnimMixer/Source/MovieSceneAnimMixer/Public/MovieSceneAnimMixerModule.h b/Engine/Plugins/Experimental/MovieSceneAnimMixer/Source/MovieSceneAnimMixer/Public/MovieSceneAnimMixerModule.h
--- a/Engine/Plugins/Experimental/MovieSceneAnimMixer/Source/MovieSceneAnimMixer/Public/MovieSceneAnimMixerModule.h
+++ b/Engine/Plugins/Experimental/MovieSceneAnimMixer/Source/MovieSceneAnimMixer/Public/MovieSceneAnimMixerModule.h
@@ -3,6 +3,7 @@
 #pragma once
 
 #include "CoreMinimal.h"
+#include "IAnimGraphRuntime_SequencerMixerTargetConnector.h"
 #include "Modules/ModuleInterface.h"
 #include "HAL/IConsoleManager.h"
 
@@ -23,11 +24,14 @@ namespace UE::MovieScene
 	
 	extern MOVIESCENEANIMMIXER_API FAutoConsoleVariableRef CVarDefaultAnimNextInjectionSite;
 
-class FMovieSceneAnimMixerModule : public IModuleInterface
+class FMovieSceneAnimMixerModule : public IModuleInterface, public IAnimGraphRuntime_SequencerMixerTargetConnector
 {
 	// IModuleInterface interface
 	MOVIESCENEANIMMIXER_API virtual void StartupModule() override;
 	MOVIESCENEANIMMIXER_API virtual void ShutdownModule() override;
+	
+	virtual void ApplySequencerMixedPose(FPoseContext& Output, FName InTargetName) const override;
+
 };
 
 }
diff --git a/Engine/Plugins/Experimental/MovieSceneAnimMixer/Source/MovieSceneAnimMixer/Public/Systems/MovieSceneAnimBlueprintTargetSystem.h b/Engine/Plugins/Experimental/MovieSceneAnimMixer/Source/MovieSceneAnimMixer/Public/Systems/MovieSceneAnimBlueprintTargetSystem.h
--- a/Engine/Plugins/Experimental/MovieSceneAnimMixer/Source/MovieSceneAnimMixer/Public/Systems/MovieSceneAnimBlueprintTargetSystem.h
+++ b/Engine/Plugins/Experimental/MovieSceneAnimMixer/Source/MovieSceneAnimMixer/Public/Systems/MovieSceneAnimBlueprintTargetSystem.h
@@ -25,6 +25,12 @@ struct MOVIESCENEANIMMIXER_API FMovieSceneAnimBlueprintTarget : public FMovieSce
 		return HashCombine(GetTypeHash(FMovieSceneAnimBlueprintTarget::StaticStruct()), GetTypeHash(Target.BlueprintNodeName));
 	}
 
+	FMovieSceneAnimBlueprintTarget(FName InBlueprintNodeName)
+		: BlueprintNodeName(InBlueprintNodeName)
+	{
+		
+	}
+
 	FMovieSceneAnimBlueprintTarget()
 		: BlueprintNodeName(FAnimNode_SequencerMixerTarget::DefaultTargetName)
 	{
diff --git a/Engine/Plugins/Experimental/MovieSceneAnimMixer/Source/MovieSceneAnimMixer/Public/Systems/MovieSceneAnimMixerSystem.h b/Engine/Plugins/Experimental/MovieSceneAnimMixer/Source/MovieSceneAnimMixer/Public/Systems/MovieSceneAnimMixerSystem.h
--- a/Engine/Plugins/Experimental/MovieSceneAnimMixer/Source/MovieSceneAnimMixer/Public/Systems/MovieSceneAnimMixerSystem.h
+++ b/Engine/Plugins/Experimental/MovieSceneAnimMixer/Source/MovieSceneAnimMixer/Public/Systems/MovieSceneAnimMixerSystem.h
@@ -13,6 +13,7 @@
 #include "EvaluationVM/Tasks/BlendKeyframes.h"
 #include "MovieSceneAnimMixerSystem.generated.h"
 
+struct FMovieSceneSkeletalAnimationComponentData;
 struct FAnimNextEvaluationTask;
 
 namespace UE::UAF
@@ -145,7 +146,7 @@ public:
 
 	UMovieSceneAnimMixerSystem(const FObjectInitializer& ObjInit);
 
-	static TInstancedStruct<FMovieSceneMixedAnimationTarget> ResolveAnimationTarget(FObjectKey ObjectKey, const TInstancedStruct<FMovieSceneMixedAnimationTarget>& InTarget);
+	static TInstancedStruct<FMovieSceneMixedAnimationTarget> ResolveAnimationTarget(FObjectKey ObjectKey, const TInstancedStruct<FMovieSceneMixedAnimationTarget>& InTarget, const FMovieSceneSkeletalAnimationComponentData* SkelAnimData = nullptr);
 
 	TSharedPtr<FMovieSceneMixerRootMotionComponentData> FindRootMotion(FObjectKey InObject) const;
 	void AssignRootMotion(FObjectKey InObjectKey, TSharedPtr<FMovieSceneMixerRootMotionComponentData> RootMotion);
diff --git a/Engine/Plugins/Experimental/MovieSceneAnimMixer/Source/MovieSceneAnimMixerEditor/Private/MovieSceneAnimMixerEditorModule.cpp b/Engine/Plugins/Experimental/MovieSceneAnimMixer/Source/MovieSceneAnimMixerEditor/Private/MovieSceneAnimMixerEditorModule.cpp
--- a/Engine/Plugins/Experimental/MovieSceneAnimMixer/Source/MovieSceneAnimMixerEditor/Private/MovieSceneAnimMixerEditorModule.cpp
+++ b/Engine/Plugins/Experimental/MovieSceneAnimMixer/Source/MovieSceneAnimMixerEditor/Private/MovieSceneAnimMixerEditorModule.cpp
@@ -2,6 +2,8 @@
 
 #include "MovieSceneAnimMixerEditorModule.h"
 
+#include "AnimBlueprintExtension_SequencerMixerTarget.h"
+#include "IAnimGraph_SequencerMixerTargetConnector.h"
 #include "ISequencerModule.h"
 #include "Modules/ModuleInterface.h"
 #include "Modules/ModuleManager.h"
@@ -84,6 +86,8 @@ namespace UE::MovieScene
 	{
 		using namespace UE::Sequencer;
 
+		IModularFeatures::Get().RegisterModularFeature(IAnimGraph_SequencerMixerTargetConnector::GetModularFeatureName(), this);
+
 		ISequencerModule& SequencerModule = FModuleManager::Get().LoadModuleChecked<ISequencerModule>(SequencerModuleName.Resolve());
 		AnimationTrackEditorHandle = SequencerModule.RegisterTrackEditor(FOnCreateTrackEditor::CreateStatic( &FAnimationMixerTrackEditor::CreateTrackEditor));
 
@@ -96,6 +100,13 @@ namespace UE::MovieScene
 		{
 			SequencerModule->UnregisterTrackModel(AnimationTrackEditorHandle);
 		}
+
+		IModularFeatures::Get().UnregisterModularFeature(IAnimGraph_SequencerMixerTargetConnector::GetModularFeatureName(), this);
+	}
+
+	void FMovieSceneAnimMixerEditorModule::GetSequencerMixerTargetRequiredExtensions(TArray<TSubclassOf<UAnimBlueprintExtension>>& OutExtensions) const
+	{
+		OutExtensions.Add(UAnimBlueprintExtension_SequencerMixerTarget::StaticClass());
 	}
 }
 
diff --git a/Engine/Plugins/Experimental/MovieSceneAnimMixer/Source/MovieSceneAnimMixerEditor/Private/MovieSceneAnimMixerEditorModule.h b/Engine/Plugins/Experimental/MovieSceneAnimMixer/Source/MovieSceneAnimMixerEditor/Private/MovieSceneAnimMixerEditorModule.h
--- a/Engine/Plugins/Experimental/MovieSceneAnimMixer/Source/MovieSceneAnimMixerEditor/Private/MovieSceneAnimMixerEditorModule.h
+++ b/Engine/Plugins/Experimental/MovieSceneAnimMixer/Source/MovieSceneAnimMixerEditor/Private/MovieSceneAnimMixerEditorModule.h
@@ -3,17 +3,20 @@
 #pragma once
 
 #include "CoreMinimal.h"
+#include "IAnimGraph_SequencerMixerTargetConnector.h"
 #include "Modules/ModuleInterface.h"
 
 namespace UE::MovieScene
 {
-	class FMovieSceneAnimMixerEditorModule : public IModuleInterface
+	class FMovieSceneAnimMixerEditorModule : public IModuleInterface, public IAnimGraph_SequencerMixerTargetConnector
 	{
 	public:
-
+		
 		// IModuleInterface interface
 		virtual void StartupModule() override;
 		virtual void ShutdownModule() override;
+		
+		virtual void GetSequencerMixerTargetRequiredExtensions(TArray<TSubclassOf<UAnimBlueprintExtension>>& OutExtensions) const override;
 
 	private:
 		FDelegateHandle AnimationTrackEditorHandle;
diff --git a/Engine/Plugins/Experimental/MovieScenePoseSearchTracks/Source/MovieScenePoseSearchTracks/Private/Sections/MovieSceneStitchAnimSection.cpp b/Engine/Plugins/Experimental/MovieScenePoseSearchTracks/Source/MovieScenePoseSearchTracks/Private/Sections/MovieSceneStitchAnimSection.cpp
--- a/Engine/Plugins/Experimental/MovieScenePoseSearchTracks/Source/MovieScenePoseSearchTracks/Private/Sections/MovieSceneStitchAnimSection.cpp
+++ b/Engine/Plugins/Experimental/MovieScenePoseSearchTracks/Source/MovieScenePoseSearchTracks/Private/Sections/MovieSceneStitchAnimSection.cpp
@@ -395,8 +395,6 @@ void UMovieSceneStitchAnimSection::ImportEntityImpl(UMovieSceneEntitySystemLinke
 		.Add(BuiltInComponents->BoundObjectResolver, UMovieSceneSkeletalAnimationSystem::ResolveSkeletalMeshComponentBinding)
 		.AddConditional(BuiltInComponents->WeightChannel, &Weight, Weight.HasAnyData())
 		.AddTag(TrackComponents->Tags.AnimMixerPoseProducer)
-		.Add(AnimMixerComponents->Priority, MixedAnimationPriority)
-		.Add(AnimMixerComponents->Target, MixedAnimationTarget)
 		.Add(AnimMixerComponents->Task, EvalTask)
 		.Add(AnimMixerComponents->RootMotionSettings, RootMotionSettings)
 		.AddTag(AnimMixerComponents->Tags.RequiresBlending)
diff --git a/Engine/Plugins/Experimental/MovieScenePoseSearchTracks/Source/MovieScenePoseSearchTracks/Public/Sections/MovieSceneStitchAnimSection.h b/Engine/Plugins/Experimental/MovieScenePoseSearchTracks/Source/MovieScenePoseSearchTracks/Public/Sections/MovieSceneStitchAnimSection.h
--- a/Engine/Plugins/Experimental/MovieScenePoseSearchTracks/Source/MovieScenePoseSearchTracks/Public/Sections/MovieSceneStitchAnimSection.h
+++ b/Engine/Plugins/Experimental/MovieScenePoseSearchTracks/Source/MovieScenePoseSearchTracks/Public/Sections/MovieSceneStitchAnimSection.h
@@ -8,7 +8,6 @@
 #include "Channels/MovieSceneFloatChannel.h"
 #include "EntitySystem/IMovieSceneEntityProvider.h"
 #include "StructUtils/InstancedStruct.h"
-#include "MovieSceneMixedAnimationTarget.h"
 #include "MovieSceneAnimationMixerTrack.h"
 #include "Animation/AnimationAsset.h"
 #include "EvaluationVM/EvaluationTask.h"
@@ -131,12 +130,6 @@ public:
 	UPROPERTY(EditAnywhere, Category = "Stitch")
 	FTransform TargetTransform;
 
-	UPROPERTY(EditAnywhere, Category = "Mixed Animation")
-	TInstancedStruct<FMovieSceneMixedAnimationTarget> MixedAnimationTarget;
-
-	UPROPERTY(EditAnywhere, Category = "Mixed Animation")
-	int32 MixedAnimationPriority = 0;
-
 	/** The weight curve for this animation section */
 	UPROPERTY()
 	FMovieSceneFloatChannel Weight;
diff --git a/Engine/Source/Editor/AnimGraph/Private/AnimGraphNode_Slot.cpp b/Engine/Source/Editor/AnimGraph/Private/AnimGraphNode_Slot.cpp
--- a/Engine/Source/Editor/AnimGraph/Private/AnimGraphNode_Slot.cpp
+++ b/Engine/Source/Editor/AnimGraph/Private/AnimGraphNode_Slot.cpp
@@ -1,6 +1,8 @@
 // Copyright Epic Games, Inc. All Rights Reserved.
 
 #include "AnimGraphNode_Slot.h"
+#include "IAnimGraph_SequencerMixerTargetConnector.h"
+#include "Features/IModularFeatures.h"
 
 
 /////////////////////////////////////////////////////
@@ -76,4 +78,14 @@ void UAnimGraphNode_Slot::BakeDataDuringCompilation(class FCompilerResultsLog& M
 	}
 }
 
+void UAnimGraphNode_Slot::GetRequiredExtensions(TArray<TSubclassOf<UAnimBlueprintExtension>>& OutExtensions) const
+{
+	Super::GetRequiredExtensions(OutExtensions);
+	
+	if (const IAnimGraph_SequencerMixerTargetConnector* MixerTargetConnector = IAnimGraph_SequencerMixerTargetConnector::Get())
+	{
+		MixerTargetConnector->GetSequencerMixerTargetRequiredExtensions(OutExtensions);
+	}
+}
+
 #undef LOCTEXT_NAMESPACE
diff --git a/Engine/Source/Editor/AnimGraph/Public/AnimGraphNode_Slot.h b/Engine/Source/Editor/AnimGraph/Public/AnimGraphNode_Slot.h
--- a/Engine/Source/Editor/AnimGraph/Public/AnimGraphNode_Slot.h
+++ b/Engine/Source/Editor/AnimGraph/Public/AnimGraphNode_Slot.h
@@ -26,6 +26,8 @@ class UAnimGraphNode_Slot : public UAnimGraphNode_Base
 	//~ Begin UAnimGraphNode_Base Interface
 	virtual FString GetNodeCategory() const override;
 	virtual void BakeDataDuringCompilation(class FCompilerResultsLog& MessageLog) override;
+	virtual void GetRequiredExtensions(TArray<TSubclassOf<UAnimBlueprintExtension>>& OutExtensions) const override;
+
 	//~ End UAnimGraphNode_Base Interface
 
 private:
diff --git a/Engine/Source/Editor/AnimGraph/Public/IAnimGraph_SequencerMixerTargetConnector.h b/Engine/Source/Editor/AnimGraph/Public/IAnimGraph_SequencerMixerTargetConnector.h
new file mode 100644
--- /dev/null
+++ b/Engine/Source/Editor/AnimGraph/Public/IAnimGraph_SequencerMixerTargetConnector.h
@@ -0,0 +1,32 @@
+// Copyright Epic Games, Inc. All Rights Reserved.
+
+#pragma once
+
+#include "Features/IModularFeatures.h"
+#include "Templates/SubclassOf.h"
+
+class UAnimBlueprintExtension;
+
+class ANIMGRAPH_API IAnimGraph_SequencerMixerTargetConnector : public IModularFeature
+{
+public:
+
+	static FName GetModularFeatureName()
+	{
+		static FName FeatureName = FName(TEXT("AnimGraph_SequencerMixerTargetConnector"));
+		return FeatureName;
+	}
+
+	static const IAnimGraph_SequencerMixerTargetConnector* Get() 
+	{
+		if (IModularFeatures::Get().IsModularFeatureAvailable(IAnimGraph_SequencerMixerTargetConnector::GetModularFeatureName()))
+		{
+			return &IModularFeatures::Get().GetModularFeature<IAnimGraph_SequencerMixerTargetConnector>(IAnimGraph_SequencerMixerTargetConnector::GetModularFeatureName());
+		}
+		return nullptr;
+	}
+
+	virtual ~IAnimGraph_SequencerMixerTargetConnector() = default;
+	
+	virtual void GetSequencerMixerTargetRequiredExtensions(TArray<TSubclassOf<UAnimBlueprintExtension>>& OutExtensions) const = 0;
+};
\ No newline at end of file
diff --git a/Engine/Source/Runtime/AnimGraphRuntime/Private/AnimNodes/AnimNode_Slot.cpp b/Engine/Source/Runtime/AnimGraphRuntime/Private/AnimNodes/AnimNode_Slot.cpp
--- a/Engine/Source/Runtime/AnimGraphRuntime/Private/AnimNodes/AnimNode_Slot.cpp
+++ b/Engine/Source/Runtime/AnimGraphRuntime/Private/AnimNodes/AnimNode_Slot.cpp
@@ -7,6 +7,7 @@
 #include "Animation/AnimStats.h"
 #include "Animation/AnimTrace.h"
 #include "Animation/AnimNode_Inertialization.h"
+#include "IAnimGraphRuntime_SequencerMixerTargetConnector.h"
 
 #include UE_INLINE_GENERATED_CPP_BY_NAME(AnimNode_Slot)
 
@@ -107,6 +108,12 @@ void FAnimNode_Slot::Evaluate_AnyThread(FPoseContext & Output)
 		checkSlow(!Output.ContainsNaN());
 		checkSlow(Output.IsNormalized());
 	}
+
+	// Backwards compatibility with Sequencer mixer
+	if (IModularFeatures::Get().IsModularFeatureAvailable(IAnimGraphRuntime_SequencerMixerTargetConnector::GetModularFeatureName()))
+	{
+		IModularFeatures::Get().GetModularFeature<IAnimGraphRuntime_SequencerMixerTargetConnector>(IAnimGraphRuntime_SequencerMixerTargetConnector::GetModularFeatureName()).ApplySequencerMixedPose(Output, SlotName);
+	}
 }
 
 void FAnimNode_Slot::GatherDebugData(FNodeDebugData& DebugData)
diff --git a/Engine/Source/Runtime/AnimGraphRuntime/Public/IAnimGraphRuntime_SequencerMixerTargetConnector.h b/Engine/Source/Runtime/AnimGraphRuntime/Public/IAnimGraphRuntime_SequencerMixerTargetConnector.h
new file mode 100644
--- /dev/null
+++ b/Engine/Source/Runtime/AnimGraphRuntime/Public/IAnimGraphRuntime_SequencerMixerTargetConnector.h
@@ -0,0 +1,31 @@
+﻿// Copyright Epic Games, Inc. All Rights Reserved.
+
+#pragma once
+
+#include "Features/IModularFeatures.h"
+
+struct FPoseContext;
+
+class ANIMGRAPHRUNTIME_API IAnimGraphRuntime_SequencerMixerTargetConnector : public IModularFeature
+{
+public:
+
+	static FName GetModularFeatureName()
+	{
+		static FName FeatureName = FName(TEXT("AnimGraphRuntime_SequencerMixerTargetConnector"));
+		return FeatureName;
+	}
+
+	static const IAnimGraphRuntime_SequencerMixerTargetConnector* Get()
+	{
+		if (IModularFeatures::Get().IsModularFeatureAvailable(IAnimGraphRuntime_SequencerMixerTargetConnector::GetModularFeatureName()))
+		{
+			return &IModularFeatures::Get().GetModularFeature<IAnimGraphRuntime_SequencerMixerTargetConnector>(IAnimGraphRuntime_SequencerMixerTargetConnector::GetModularFeatureName());
+		}
+		return nullptr;
+	}
+
+	virtual ~IAnimGraphRuntime_SequencerMixerTargetConnector() = default;
+	
+	virtual void ApplySequencerMixedPose(FPoseContext& Output, FName InTargetName) const = 0;
+};
\ No newline at end of file
diff --git a/Engine/Source/Runtime/MovieSceneTracks/Private/Sections/MovieSceneSkeletalAnimationSection.cpp b/Engine/Source/Runtime/MovieSceneTracks/Private/Sections/MovieSceneSkeletalAnimationSection.cpp
--- a/Engine/Source/Runtime/MovieSceneTracks/Private/Sections/MovieSceneSkeletalAnimationSection.cpp
+++ b/Engine/Source/Runtime/MovieSceneTracks/Private/Sections/MovieSceneSkeletalAnimationSection.cpp
@@ -130,8 +130,6 @@ UMovieSceneSkeletalAnimationSection::UMovieSceneSkeletalAnimationSection( const
 	bMatchRotationPitch = false;
 	bMatchIncludeZHeight = false;
 	bDebugForceTickPose = false;
-
-	MixedAnimationTarget = TInstancedStruct<FMovieSceneMixedAnimationTarget>::Make();
 }
 
 UAnimSequenceBase* UMovieSceneSkeletalAnimationSection::GetAnimation() const
diff --git a/Engine/Source/Runtime/MovieSceneTracks/Public/Sections/MovieSceneSkeletalAnimationSection.h b/Engine/Source/Runtime/MovieSceneTracks/Public/Sections/MovieSceneSkeletalAnimationSection.h
--- a/Engine/Source/Runtime/MovieSceneTracks/Public/Sections/MovieSceneSkeletalAnimationSection.h
+++ b/Engine/Source/Runtime/MovieSceneTracks/Public/Sections/MovieSceneSkeletalAnimationSection.h
@@ -9,7 +9,6 @@
 #include "Channels/MovieSceneFloatChannel.h"
 #include "EntitySystem/IMovieSceneEntityProvider.h"
 #include "StructUtils/InstancedStruct.h"
-#include "MovieSceneMixedAnimationTarget.h"
 #include "MovieSceneSkeletalAnimationSection.generated.h"
 
 struct FMovieSceneSkeletalAnimRootMotionTrackParams;
@@ -249,12 +248,6 @@ public:
 	uint8 bShowSkeleton : 1;
 #endif
 
-	UPROPERTY(EditAnywhere, Category = "Mixed Animation")
-	TInstancedStruct<FMovieSceneMixedAnimationTarget> MixedAnimationTarget;
-
-	UPROPERTY(EditAnywhere, Category = "Mixed Animation")
-	int32 MixedAnimationPriority = 0;
-
 	//Previous transform used to specify the global OffsetTransform while calculting the root motions.
 	FTransform PreviousTransform;
 
