Sequencer: Temporary fix to stitch track crash.
We had removed these components from the section entity since ideally the anim mixer track or a decoration would be responsible for them, but for now they're needed to get the mixer to pick up the entity properly and create a mixer entry for it. We'll figure out a better solution as part of the larger refactor later.

diff --git a/Engine/Plugins/Experimental/MovieSceneAnimMixer/Source/MovieSceneAnimMixer/Private/Systems/MovieSceneAnimMixerSystem.cpp b/Engine/Plugins/Experimental/MovieSceneAnimMixer/Source/MovieSceneAnimMixer/Private/Systems/MovieSceneAnimMixerSystem.cpp
--- a/Engine/Plugins/Experimental/MovieSceneAnimMixer/Source/MovieSceneAnimMixer/Private/Systems/MovieSceneAnimMixerSystem.cpp
+++ b/Engine/Plugins/Experimental/MovieSceneAnimMixer/Source/MovieSceneAnimMixer/Private/Systems/MovieSceneAnimMixerSystem.cpp
@@ -59,11 +59,14 @@ namespace UE::MovieScene
 			const FMovieSceneRootMotionSettings* RootMotionSettings,
 			const double* WeightAndEasings)
 		{
-			MixerEntry->EvalTask = Task;
-			MixerEntry->PoseWeight = WeightAndEasings ? *WeightAndEasings : 1.0;
-			if (RootMotionSettings)
+			if (MixerEntry)
 			{
-				MixerEntry->RootMotionSettings = *RootMotionSettings;
+				MixerEntry->EvalTask = Task;
+				MixerEntry->PoseWeight = WeightAndEasings ? *WeightAndEasings : 1.0;
+				if (RootMotionSettings)
+				{
+					MixerEntry->RootMotionSettings = *RootMotionSettings;
+				}
 			}
 		}
 	};
diff --git a/Engine/Plugins/Experimental/MovieScenePoseSearchTracks/Source/MovieScenePoseSearchTracks/Private/Sections/MovieSceneStitchAnimSection.cpp b/Engine/Plugins/Experimental/MovieScenePoseSearchTracks/Source/MovieScenePoseSearchTracks/Private/Sections/MovieSceneStitchAnimSection.cpp
--- a/Engine/Plugins/Experimental/MovieScenePoseSearchTracks/Source/MovieScenePoseSearchTracks/Private/Sections/MovieSceneStitchAnimSection.cpp
+++ b/Engine/Plugins/Experimental/MovieScenePoseSearchTracks/Source/MovieScenePoseSearchTracks/Private/Sections/MovieSceneStitchAnimSection.cpp
@@ -398,6 +398,9 @@ void UMovieSceneStitchAnimSection::ImportEntityImpl(UMovieSceneEntitySystemLinke
 		.Add(AnimMixerComponents->Task, EvalTask)
 		.Add(AnimMixerComponents->RootMotionSettings, RootMotionSettings)
 		.AddTag(AnimMixerComponents->Tags.RequiresBlending)
+		.Add(AnimMixerComponents->Target, TInstancedStruct<FMovieSceneMixedAnimationTarget>::Make())
+		.Add(AnimMixerComponents->MixerEntry, nullptr)
+		.Add(AnimMixerComponents->Priority, 0)
 	);
 }
 
