Moving some MovieSceneAnimMixer files from internal to public.

#rnx

diff --git a/Engine/Plugins/Experimental/MovieSceneAnimMixer/Source/MovieSceneAnimMixer/Internal/MovieSceneAnimationMixerTrack.h b/Engine/Plugins/Experimental/MovieSceneAnimMixer/Source/MovieSceneAnimMixer/Internal/MovieSceneAnimationMixerTrack.h
deleted file mode 100644
diff --git a/Engine/Plugins/Experimental/MovieSceneAnimMixer/Source/MovieSceneAnimMixer/Internal/MovieSceneRootMotionSection.h b/Engine/Plugins/Experimental/MovieSceneAnimMixer/Source/MovieSceneAnimMixer/Internal/MovieSceneRootMotionSection.h
deleted file mode 100644
diff --git a/Engine/Plugins/Experimental/MovieSceneAnimMixer/Source/MovieSceneAnimMixer/Public/MovieSceneAnimationMixerTrack.h b/Engine/Plugins/Experimental/MovieSceneAnimMixer/Source/MovieSceneAnimMixer/Public/MovieSceneAnimationMixerTrack.h
new file mode 100644
--- /dev/null
+++ b/Engine/Plugins/Experimental/MovieSceneAnimMixer/Source/MovieSceneAnimMixer/Public/MovieSceneAnimationMixerTrack.h
@@ -0,0 +1,152 @@
+// Copyright Epic Games, Inc. All Rights Reserved.
+
+#pragma once
+
+#include "CoreTypes.h"
+#include "Channels/MovieSceneDoubleChannel.h"
+#include "Channels/MovieSceneByteChannel.h"
+#include "Decorations/IMovieSceneChannelDecoration.h"
+#include "Decorations/IMovieSceneSectionDecoration.h"
+#include "EntitySystem/IMovieSceneEntityDecorator.h"
+#include "MovieSceneMixedAnimationTarget.h"
+#include "MovieSceneNameableTrack.h"
+#include "StructUtils/InstancedStruct.h"
+#include "Tracks/MovieSceneCommonAnimationTrack.h"
+#include "Systems/MovieSceneRootMotionSystem.h"
+#include "MovieSceneAnimationMixerTrack.generated.h"
+
+USTRUCT()
+struct FMovieSceneByteChannelDefaultOnly : public FMovieSceneByteChannel
+{
+	GENERATED_BODY()
+
+	bool SerializeFromMismatchedTag(const FPropertyTag& Tag, FStructuredArchive::FSlot Slot);
+};
+
+template<>
+struct TStructOpsTypeTraits<FMovieSceneByteChannelDefaultOnly> : public TStructOpsTypeTraitsBase2<FMovieSceneByteChannelDefaultOnly>
+{
+	enum { WithStructuredSerializeFromMismatchedTag = true };
+};
+
+template<>
+struct TMovieSceneChannelTraits<FMovieSceneByteChannelDefaultOnly> : TMovieSceneChannelTraitsBase<FMovieSceneByteChannelDefaultOnly>
+{
+#if WITH_EDITOR
+
+	/** Byte channels can have external values (ie, they can get their values from external objects for UI purposes) */
+	typedef TMovieSceneExternalValue<uint8> ExtendedEditorDataType;
+
+#endif
+};
+
+UINTERFACE(MinimalAPI)
+class UMovieSceneAnimationSectionInterface : public UInterface
+{
+	GENERATED_BODY()
+};
+
+
+class IMovieSceneAnimationSectionInterface
+{
+public:
+	GENERATED_BODY()
+
+	virtual int32 GetRowSortOrder() const = 0;
+
+	// Allows different section types to have different colors when used in the mixer, without needing to override their colors everywhere.
+	virtual FColor GetMixerSectionTint() const = 0;
+};
+
+UCLASS(MinimalAPI)
+class UMovieSceneAnimationSectionDecoration
+	: public UMovieSceneSignedObject
+	, public IMovieSceneEntityDecorator
+	, public IMovieSceneSectionDecoration
+	, public IMovieSceneAnimationSectionInterface
+{
+	GENERATED_BODY()
+
+	virtual void ExtendEntityImpl(UMovieSceneEntitySystemLinker* EntityLinker, const FEntityImportParams& Params, FImportedEntity* OutImportedEntity) override;
+
+	virtual int32 GetRowSortOrder() const
+	{
+		return RowSortOrder;
+	}
+
+	virtual FColor GetMixerSectionTint() const override
+	{
+		return MixerTintOverride;
+	}
+
+	UPROPERTY()
+	int32 RowSortOrder = 0;
+
+private:
+	FColor MixerTintOverride = FColor();
+};
+
+UCLASS(MinimalAPI)
+class UMovieSceneAnimationBaseTransformDecoration
+	: public UMovieSceneSignedObject
+	, public IMovieSceneChannelDecoration
+	, public IMovieSceneEntityDecorator
+{
+public:
+
+	GENERATED_BODY()
+
+	UMovieSceneAnimationBaseTransformDecoration(const FObjectInitializer& ObjInit);
+
+	MOVIESCENEANIMMIXER_API EMovieSceneRootMotionTransformMode GetRootTransformMode() const;
+
+	virtual EMovieSceneChannelProxyType PopulateChannelProxy(FMovieSceneChannelProxyData& OutProxyData) override;
+	virtual void ExtendEntityImpl(UMovieSceneEntitySystemLinker* EntityLinker, const FEntityImportParams& Params, FImportedEntity* OutImportedEntity) override;
+
+	UPROPERTY()
+	FMovieSceneDoubleChannel Location[3];
+
+	UPROPERTY()
+	FMovieSceneDoubleChannel Rotation[3];
+
+	UPROPERTY(EditAnywhere, Category="Root Transform")
+	FVector RootOriginLocation;
+
+	UPROPERTY()
+	FMovieSceneByteChannelDefaultOnly RootMotionSpace;
+
+	UPROPERTY()
+	FMovieSceneByteChannelDefaultOnly TransformMode;
+};
+
+
+/**
+ *
+ */
+UCLASS(MinimalAPI)
+class UMovieSceneAnimationMixerTrack
+	: public UMovieSceneCommonAnimationTrack
+{
+	GENERATED_BODY()
+
+public:
+
+	UMovieSceneAnimationMixerTrack(const FObjectInitializer& Init);
+
+	virtual bool FixRowIndices() override;
+	virtual bool SupportsType(TSubclassOf<UMovieSceneSection> SectionClass) const override;
+
+	virtual void OnSectionAddedImpl(UMovieSceneSection* Secton) override;
+
+#if WITH_EDITORONLY_DATA
+
+	virtual FText GetTrackRowDisplayName(int32 RowIndex) const override;
+	virtual FText GetDefaultDisplayName() const override;
+	virtual bool CanRename() const override;
+#endif
+
+public:
+
+	UPROPERTY(EditAnywhere, Category = "Mixed Animation")
+	TInstancedStruct<FMovieSceneMixedAnimationTarget> MixedAnimationTarget;
+};
diff --git a/Engine/Plugins/Experimental/MovieSceneAnimMixer/Source/MovieSceneAnimMixer/Public/MovieSceneRootMotionSection.h b/Engine/Plugins/Experimental/MovieSceneAnimMixer/Source/MovieSceneAnimMixer/Public/MovieSceneRootMotionSection.h
new file mode 100644
--- /dev/null
+++ b/Engine/Plugins/Experimental/MovieSceneAnimMixer/Source/MovieSceneAnimMixer/Public/MovieSceneRootMotionSection.h
@@ -0,0 +1,60 @@
+// Copyright Epic Games, Inc. All Rights Reserved.
+
+#pragma once
+
+#include "CoreTypes.h"
+#include "Channels/MovieSceneByteChannel.h"
+#include "EntitySystem/IMovieSceneEntityProvider.h"
+#include "MovieSceneSection.h"
+#include "MovieSceneAnimationMixerTrack.h"
+#include "MovieSceneRootMotionSection.generated.h"
+
+UENUM()
+enum class EMovieSceneRootMotionDestination : uint8
+{
+	/** Throw away any transform on the root bone */
+	Discard,
+	/** Leave the root bone with whatever transform it ended up with after evaluation */
+	RootBone,
+	/** Copy the root bone's transform onto the owning Component, and reset the root transform */
+	Component,
+	/** Copy the root bone's transform onto the owning Actor, and reset the root transform */
+	Actor,
+	/** Leave the root motion transform on an attribute for external systems to pick up */
+	Attribute,
+};
+
+/**
+ * 
+ */
+UCLASS(MinimalAPI, DisplayName="Root Motion")
+class UMovieSceneRootMotionSection
+	: public UMovieSceneSection
+	, public IMovieSceneEntityProvider
+	, public IMovieSceneAnimationSectionInterface
+{
+	GENERATED_BODY()
+
+public:
+
+	UMovieSceneRootMotionSection(const FObjectInitializer& Init);
+
+	virtual int32 GetRowSortOrder() const;
+
+	virtual FColor GetMixerSectionTint() const override
+	{
+		return MixerTintOverride;
+	}
+
+	virtual EMovieSceneChannelProxyType CacheChannelProxy() override;
+
+	virtual void ImportEntityImpl(UMovieSceneEntitySystemLinker* EntityLinker, const FEntityImportParams& Params, FImportedEntity* OutImportedEntity) override;
+	virtual bool PopulateEvaluationFieldImpl(const TRange<FFrameNumber>& EffectiveRange, const FMovieSceneEvaluationFieldEntityMetaData& InMetaData, FMovieSceneEntityComponentFieldBuilder* OutFieldBuilder) override;
+
+private:
+
+	UPROPERTY()
+	FMovieSceneByteChannel RootDestinationChannel;
+
+	FColor MixerTintOverride = FColor(20, 70, 70, 200);
+};
diff --git a/Engine/Plugins/Experimental/MovieSceneAnimMixer/Source/MovieSceneAnimMixer/Public/Systems/MovieSceneRootMotionSystem.h b/Engine/Plugins/Experimental/MovieSceneAnimMixer/Source/MovieSceneAnimMixer/Public/Systems/MovieSceneRootMotionSystem.h
--- a/Engine/Plugins/Experimental/MovieSceneAnimMixer/Source/MovieSceneAnimMixer/Public/Systems/MovieSceneRootMotionSystem.h
+++ b/Engine/Plugins/Experimental/MovieSceneAnimMixer/Source/MovieSceneAnimMixer/Public/Systems/MovieSceneRootMotionSystem.h
@@ -4,7 +4,6 @@
 
 #include "EntitySystem/MovieSceneEntitySystem.h"
 #include "EntitySystem/MovieSceneEntityIDs.h"
-#include "MovieSceneAnimationMixerTrack.h"
 #include "AnimSequencerInstanceProxy.h"
 #include "EvaluationVM/EvaluationTask.h"
 #include "Tickable.h"
@@ -15,6 +14,25 @@ enum class EMovieSceneRootMotionDestination : uint8;
 class UMovieSceneAnimMixerSystem;
 class UMovieSceneRootMotionSystem;
 
+// Enum describing what space that root motion should be applied in.
+UENUM()
+enum class EMovieSceneRootMotionSpace : uint8
+{
+	// Root motion should be applied in animation space, meaning that it will be applied on top of the blended transform track and transform origin.
+	AnimationSpace,
+	// Root motion should be applied in world space, meaning that it will override any transform track or transform origin.
+	WorldSpace
+};
+
+
+// Enum describing what space that root motion should be applied in.
+UENUM()
+enum class EMovieSceneRootMotionTransformMode : uint8
+{
+	Asset,
+	Offset,
+	Override,
+};
 
 USTRUCT()
 struct FSkeletalAnimationRootMotionOverride
