Sequencer: Large anim mixer update and Sequencer MVVM refactor.

1) Refactor view model track/track row code to not be as hardcoded to check for FTrackModel/FTrackRowModel but instead use extensions where appropriate. In addition to ITrackExtension, further extensions ITrackRowExtension, ISectionOwnerExtension, and ITopLevelChannelHolderExtension have been added. Some functionality previously present in ITrackExtension has been broken out to others of these extensions. The largest part of this changelist is updating code that previously looked for FTrackModel/FTrackRowModel or ITrackExtension to use the appropriate new extension. This refactor is a foundational piece for introducing custom layout code for the anim mixer layers which can act at times both like a track and/or a track row.

Some of the layout code that was in FTrackModel has been moved to FTrackRowModel for better modularity.

As part of this, the track model storage extension now stores ITrackExtension shared pointers rather than FTrackModel shared pointers. This will allow layers to act as 'tracks' at times and be searchable by other code.

2) Add 'ParentTrack' property to UMovieSceneTrack which allows tracks to be added to an object binding's track list that don't show under the object binding outliner items, but under other custom outliner items. Used for control rig tracks that are child tracks of the anim mixer track. The alternative was changing their outer, but so much code relies on asking the object binding for its tracks directly and this was the easier option.

3) Adds Animation Mixer Layers.

Mixer Layers are sort of like a more substantial 'track row' for the mixer. A mixer layer can contain either animation pose producing sections (e.g. skel anim section, stitch section) OR a child track such as a control rig track.

The order of the layers matter- layers with a higher index evaluate later and therefore their poses can override previous layers.

These concepts existed before the anim mixer track and layers in the mixer, but relied on magic 'priority' set on individual anim sections. Now the layers dictate the priority based on their index and the mixer track itself dictates the target.

More specifically this changelist includes:
- UMovieSceneAnimationMixerLayer class for organizing sections and child tracks into layers
- IMovieSceneAnimationSectionInterface has been replaced with IMovieSceneAnimationMixerItemInterface for mixer item UI customization (tint, display name, icon) as it is used for both sections and child tracks. In this changelist, both interfaces exist temporarily- the UI changes come in a separate changelist.
- Track decoration for extending child track entities with mixer components. Sections do the same with the existing animation section decoration.
- MixerLayer component type for sending mixer to the ECS system- will be used later to get properties off of the layer such as masks, etc.
- Layer management for UMovieSceneAnimationMixerTrack

4) Section Outliner Model and Drag-Drop Updates for Anim Mixer Tracks

- FSectionOutlinerModel: Allows individual sections to appear as outliner items with their channels. This allows multiple sections to be on a single layer but when the layer is expanded it separates the channels for each section underneath their section outliner model to avoid overlapping of channels.
- FTrackRowReorderDragDropOp: Drag-drop operation for reordering track rows/layers by dragging their outliner elements. I've currently only enabled this for anim mixer layers (next changelist), but theoretically we could enable this for regular track rows if we wanted.
- Section vertical drag now delegates to track models to allow the anim mixer track model to override this.
- I've taken advantage of this to allow section dragging to work differently for the anim mixer (with preview insertion painting and not moving the section until the drag is completed)
- Paint callbacks for drop preview indicators during drag operations

5) Implements MVVM editor-side code for the new animation mixer track model and animation mixer layer model. This takes advantage of previous refactoring of track model, drag/drop, etc. in order to properly layout animation mixer tracks, their layers, and their child tracks and sections.

6) Menu changes to allow control rig and stitch tracks to properly be added to the anim mixer track and allowing future item types to be added in a more modular way. This includes the '+' button menu on Animation Mixer track and on Animation Layer tracks and provides ways for other track editors/tracks/sections to register themselves for being included in the mixer track. As part of this, when the Anim Mixer plugin is enabled, the Animation track submenu has been reorganized- Control Rig is now underneath 'Animation' as is Stitch and Anim Mixer.

#Synced-CL 48992089

diff --git a/Engine/Plugins/Animation/ControlRig/Source/ControlRigEditor/Private/AnimDetails/AnimDetailsSelection.cpp b/Engine/Plugins/Animation/ControlRig/Source/ControlRigEditor/Private/AnimDetails/AnimDetailsSelection.cpp
--- a/Engine/Plugins/Animation/ControlRig/Source/ControlRigEditor/Private/AnimDetails/AnimDetailsSelection.cpp
+++ b/Engine/Plugins/Animation/ControlRig/Source/ControlRigEditor/Private/AnimDetails/AnimDetailsSelection.cpp
@@ -433,7 +433,7 @@ void UAnimDetailsSelection::PropagateSelectionToCurveEditor()
 	struct FOutlinerExtensionData
 	{
 		UMovieSceneControlRigParameterTrack* ControlRigParameterTrack = nullptr;
-		TSharedPtr<FTrackModel> TrackModel;
+		TViewModelPtr<ITrackExtension> TrackModel;
 		TViewModelPtr<FChannelGroupOutlinerModel> ChannelModel;
 		TViewModelPtr<ICurveEditorTreeItemExtension> CurveEditorItem;
 
@@ -444,7 +444,7 @@ void UAnimDetailsSelection::PropagateSelectionToCurveEditor()
 	TParentFirstChildIterator<IOutlinerExtension> CachedOutlinerExtenstionIt = OutlinerRootItem->GetDescendantsOfType<IOutlinerExtension>();
 	for (TParentFirstChildIterator<IOutlinerExtension> OutlinerExtenstionIt = CachedOutlinerExtenstionIt; OutlinerExtenstionIt; ++OutlinerExtenstionIt)
 	{
-		if (const TSharedPtr<FTrackModel> TrackModel = OutlinerExtenstionIt.GetCurrentItem()->FindAncestorOfType<FTrackModel>())
+		if (const TViewModelPtr<ITrackExtension> TrackModel = OutlinerExtenstionIt.GetCurrentItem()->FindAncestorOfType<ITrackExtension>())
 		{
 			FOutlinerExtensionData ExtensionData;
 			ExtensionData.TrackModel = TrackModel;
diff --git a/Engine/Plugins/Animation/ControlRig/Source/ControlRigEditor/Private/AnimDetails/Proxies/AnimDetailsProxyBase.cpp b/Engine/Plugins/Animation/ControlRig/Source/ControlRigEditor/Private/AnimDetails/Proxies/AnimDetailsProxyBase.cpp
--- a/Engine/Plugins/Animation/ControlRig/Source/ControlRigEditor/Private/AnimDetails/Proxies/AnimDetailsProxyBase.cpp
+++ b/Engine/Plugins/Animation/ControlRig/Source/ControlRigEditor/Private/AnimDetails/Proxies/AnimDetailsProxyBase.cpp
@@ -61,7 +61,7 @@ namespace UE::ControlRigEditor
 			const FViewModelPtr RootModel = EditorViewModel.IsValid() ? EditorViewModel->GetRootModel() : nullptr;
 			const FSectionModelStorageExtension* SectionModelStorage = RootModel.IsValid() ? RootModel->CastDynamic<FSectionModelStorageExtension>() : nullptr;
 			const TSharedPtr<FSectionModel> SectionHandle = SectionModelStorage ? SectionModelStorage->FindModelForSection(Section) : nullptr;
-			const TSharedPtr<FViewModel> ViewModel = SectionHandle.IsValid() ? SectionHandle->GetParentTrackModel().AsModel() : nullptr;
+			const TSharedPtr<FViewModel> ViewModel = SectionHandle.IsValid() ? SectionHandle->GetParentSectionOwnerModel().AsModel() : nullptr;
 			if (!EditorViewModel.IsValid() ||
 				!RootModel.IsValid() ||
 				!SectionModelStorage || 
@@ -89,7 +89,7 @@ namespace UE::ControlRigEditor
 				}
 			}
 
-			const TSharedPtr<FTrackModel> TrackModel = SectionHandle->FindAncestorOfType<FTrackModel>();
+			const TViewModelPtr<ITrackExtension> TrackModel = SectionHandle->FindAncestorOfType<ITrackExtension>();
 			FAddKeyOperation::FromKeyAreas(TrackModel->GetTrackEditor().Get(), KeyAreas).Commit(Time, *Sequencer);
 		}
 
diff --git a/Engine/Plugins/Animation/ControlRig/Source/ControlRigEditor/Private/Sequencer/ControlRigParameterTrackEditor.cpp b/Engine/Plugins/Animation/ControlRig/Source/ControlRigEditor/Private/Sequencer/ControlRigParameterTrackEditor.cpp
--- a/Engine/Plugins/Animation/ControlRig/Source/ControlRigEditor/Private/Sequencer/ControlRigParameterTrackEditor.cpp
+++ b/Engine/Plugins/Animation/ControlRig/Source/ControlRigEditor/Private/Sequencer/ControlRigParameterTrackEditor.cpp
@@ -3,6 +3,7 @@
 #include "Sequencer/ControlRigParameterTrackEditor.h"
 #include "Sequencer/MovieSceneControlRigSystem.h"
 #include "MVVM/Extensions/ITrackExtension.h"
+#include "MVVM/Extensions/ISectionOwnerExtension.h"
 #include "MVVM/ViewModels/OutlinerViewModel.h"
 #include "MVVM/ViewModels/TrackModel.h"
 #include "MVVM/ViewModels/SectionModel.h"
@@ -120,6 +121,8 @@
 #include "Decorations/MovieSceneMuteSoloDecoration.h"
 #include "Editor/Constraints/ConstraintCreationOptions.h"
 #include "RigVMCore/RigVMMemoryStorageStruct.h"
+#include "IMovieSceneAnimMixerItemMenuProvider.h"
+#include "Tracks/MovieSceneCommonAnimationTrack.h"
 
 
 #define LOCTEXT_NAMESPACE "FControlRigParameterTrackEditor"
@@ -511,6 +514,9 @@ void FControlRigParameterTrackEditor::OnInitialize()
 
 		return false;
 	});
+
+	// Register as animation mixer item menu provider
+	IModularFeatures::Get().RegisterModularFeature(IMovieSceneAnimMixerItemMenuProvider::GetModularFeatureName(), this);
 }
 
 FControlRigParameterTrackEditor::~FControlRigParameterTrackEditor()
@@ -739,7 +745,10 @@ void FControlRigParameterTrackEditor::OnRelease()
 			OnMovieSceneChannelChangedHandle.Reset();
 		}
 	}
-	
+
+	// Unregister as animation mixer item menu provider
+	IModularFeatures::Get().UnregisterModularFeature(IMovieSceneAnimMixerItemMenuProvider::GetModularFeatureName(), this);
+
 	AcquiredResources.Release();
 }
 
@@ -1379,8 +1388,26 @@ static void IterateTracksInMovieScene(const UMovieScene& InMovieScene, TFunction
 				InCallback(CRTrack);
 			}
 		}
+
+		// Also check for Control Rig child tracks within animation tracks (e.g., Animation Mixer)
+		TArray<UMovieSceneTrack*> AllTracks = InMovieScene.FindTracks(UMovieSceneCommonAnimationTrack::StaticClass(), Binding.GetObjectGuid(), NAME_None);
+		for (UMovieSceneTrack* Track : AllTracks)
+		{
+			if (UMovieSceneCommonAnimationTrack* AnimTrack = Cast<UMovieSceneCommonAnimationTrack>(Track))
+			{
+				TArray<TObjectPtr<UMovieSceneTrack>> ChildTracks;
+				AnimTrack->GetAllChildTracks(ChildTracks);
+				for (UMovieSceneTrack* ChildTrack : ChildTracks)
+				{
+					if (UMovieSceneControlRigParameterTrack* CRTrack = Cast<UMovieSceneControlRigParameterTrack>(ChildTrack))
+					{
+						InCallback(CRTrack);
+					}
+				}
+			}
+		}
 	}
-	
+
 	for (UMovieSceneTrack* Track : InMovieScene.GetTracks())
 	{
 		if (UMovieSceneControlRigParameterTrack* CRTrack = Cast<UMovieSceneControlRigParameterTrack>(Track))
@@ -1763,12 +1790,18 @@ bool FControlRigParameterTrackEditor::FRecreateRigOperator::IsEditingObject(UObj
 }
 
 void FControlRigParameterTrackEditor::BuildObjectBindingTrackMenu(FMenuBuilder& MenuBuilder, const TArray<FGuid>& ObjectBindings, const UClass* ObjectClass)
+{
+	// Control Rig menu entries are now added via PopulateObjectBindingAnimationMenu,
+	// which is called by FCommonAnimationTrackEditor::BuildObjectBindingTrackMenu
+}
+
+void FControlRigParameterTrackEditor::PopulateObjectBindingAnimationMenu(FMenuBuilder& MenuBuilder, const TArray<FGuid>& ObjectBindings, const UClass* ObjectClass, bool bIsInsideSubmenu)
 {
 	if(!ObjectClass)
 	{
 		return;
 	}
-	
+
 	if (ObjectClass->IsChildOf(USkeletalMeshComponent::StaticClass()) ||
 		ObjectClass->IsChildOf(AActor::StaticClass()) ||
 		ObjectClass->IsChildOf(UChildActorComponent::StaticClass()))
@@ -1797,10 +1830,11 @@ void FControlRigParameterTrackEditor::BuildObjectBindingTrackMenu(FMenuBuilder&
 		if (Skeleton)
 		{
 			UMovieSceneTrack* Track = nullptr;
+			TFunction<UMovieSceneControlRigParameterTrack* (FGuid)> InTrackCreationCallback = nullptr;
 			MenuBuilder.AddSubMenu(
 				LOCTEXT("ControlRigText", "Control Rig"),
 				FText(),
-				FNewMenuDelegate::CreateSP(this, &FControlRigParameterTrackEditor::HandleAddTrackSubMenu, ObjectBindings, Track));
+				FNewMenuDelegate::CreateSP(this, &FControlRigParameterTrackEditor::HandleAddTrackSubMenu, ObjectBindings, Track, InTrackCreationCallback));
 		}
 	}
 }
@@ -1838,7 +1872,7 @@ bool FControlRigParameterTrackEditor::IsToggleFilterAssetByAnimatableControls()
 	return bFilterAssetByAnimatableControls;
 }
 
-void FControlRigParameterTrackEditor::HandleAddTrackSubMenu(FMenuBuilder& MenuBuilder, TArray<FGuid> ObjectBindings, UMovieSceneTrack* Track)
+void FControlRigParameterTrackEditor::HandleAddTrackSubMenu(FMenuBuilder& MenuBuilder, TArray<FGuid> ObjectBindings, UMovieSceneTrack* Track, TFunction<UMovieSceneControlRigParameterTrack* (FGuid)> InTrackCreationCallback)
 {
 	if (CVarEnableAdditiveControlRigs->GetBool())
 	{
@@ -1908,7 +1942,7 @@ void FControlRigParameterTrackEditor::HandleAddTrackSubMenu(FMenuBuilder& MenuBu
 			LOCTEXT("FKControlRigTooltip", "Adds the FK Control Rig"),
 			FSlateIcon(),
 			FUIAction(
-				FExecuteAction::CreateSP(this, &FControlRigParameterTrackEditor::AddFKControlRig, BoundObject, ObjectBindings[0]),
+				FExecuteAction::CreateSP(this, &FControlRigParameterTrackEditor::AddFKControlRig, BoundObject, ObjectBindings[0], InTrackCreationCallback),
 				FCanExecuteAction::CreateLambda([ExistingRigs]()
 				{
 					if (!ExistingRigs.Contains(UFKControlRig::StaticClass()))
@@ -1928,8 +1962,8 @@ void FControlRigParameterTrackEditor::HandleAddTrackSubMenu(FMenuBuilder& MenuBu
 			AssetPickerConfig.bShowTypeInColumnView = false;
 			AssetPickerConfig.bFocusSearchBoxWhenOpened = true;
 			AssetPickerConfig.bForceShowPluginContent = true;
-			AssetPickerConfig.OnAssetSelected = FOnAssetSelected::CreateRaw(this, &FControlRigParameterTrackEditor::AddControlRig, BoundObject, ObjectBindings[0]);
-			AssetPickerConfig.OnAssetEnterPressed = FOnAssetEnterPressed::CreateRaw(this, &FControlRigParameterTrackEditor::AddControlRig, BoundObject, ObjectBindings[0]);
+			AssetPickerConfig.OnAssetSelected = FOnAssetSelected::CreateRaw(this, &FControlRigParameterTrackEditor::AddControlRig, BoundObject, ObjectBindings[0], InTrackCreationCallback);
+			AssetPickerConfig.OnAssetEnterPressed = FOnAssetEnterPressed::CreateRaw(this, &FControlRigParameterTrackEditor::AddControlRig, BoundObject, ObjectBindings[0], InTrackCreationCallback);
 			AssetPickerConfig.RefreshAssetViewDelegates.Add(&RefreshControlRigPickerDelegate);
 			AssetPickerConfig.OnShouldFilterAsset = FOnShouldFilterAsset::CreateLambda([this, ExistingRigs, Skeleton](const FAssetData& AssetData)
 			{
@@ -1981,6 +2015,38 @@ void FControlRigParameterTrackEditor::HandleAddTrackSubMenu(FMenuBuilder& MenuBu
 	}
 }
 
+const UClass* FControlRigParameterTrackEditor::GetHandledMixerItemClass() const
+{
+	return UMovieSceneControlRigParameterTrack::StaticClass();
+}
+
+void FControlRigParameterTrackEditor::PopulateAddMixerItemMenu(FMenuBuilder& MenuBuilder, TArray<FGuid> ObjectBindings, UMovieSceneTrack* Track, TSharedPtr<ISequencer> Sequencer, int32 RowIndex)
+{
+	UMovieSceneCommonAnimationTrack* CommonTrack = Cast<UMovieSceneCommonAnimationTrack>(Track);
+	if (!CommonTrack)
+	{
+		return;
+	}
+
+	// Create track creation callback that adds the track as a child of the mixer track
+	auto TrackCreationCallback = [CommonTrack, RowIndex](FGuid ObjectBinding) -> UMovieSceneControlRigParameterTrack*
+	{
+		// Use the provided row index if valid, otherwise compute next available
+		int32 NewRowIndex = (RowIndex != INDEX_NONE) ? RowIndex : CommonTrack->ComputeNextAvailableRowIndex();
+		return Cast<UMovieSceneControlRigParameterTrack>(CommonTrack->AddChildTrack(ObjectBinding, UMovieSceneControlRigParameterTrack::StaticClass(), NewRowIndex));
+	};
+
+	// Wrap the control rig menu in a submenu with mixer-specific track creation
+	MenuBuilder.AddSubMenu(
+		LOCTEXT("AddControlRig", "Control Rig"),
+		LOCTEXT("AddControlRigTooltip", "Adds a control rig track with FK Control Rig option and asset picker"),
+		FNewMenuDelegate::CreateLambda([this, ObjectBindings, Track, TrackCreationCallback](FMenuBuilder& SubMenuBuilder)
+		{
+			HandleAddTrackSubMenu(SubMenuBuilder, ObjectBindings, Track, TrackCreationCallback);
+		})
+	);
+}
+
 bool FControlRigParameterTrackEditor::IsControlRigAllowed(const FAssetData& AssetData, TArray<UClass*> ExistingRigs, USkeleton* Skeleton)
 {
 	static const FName RigModuleSettingsPropertyName = GET_MEMBER_NAME_CHECKED(UControlRigBlueprint, RigModuleSettings);
@@ -2148,7 +2214,7 @@ static TArray<UMovieSceneControlRigParameterTrack*> GetExistingControlRigTracksF
 	return ExistingControlRigTracks;
 }
 
-static UMovieSceneControlRigParameterTrack* AddControlRig(TSharedPtr<ISequencer> SharedSequencer , UMovieSceneSequence* Sequence, const UClass* InClass, FGuid ObjectBinding, UControlRig* InExistingControlRig, bool bIsAdditiveControlRig)
+static UMovieSceneControlRigParameterTrack* AddControlRig(TSharedPtr<ISequencer> SharedSequencer , UMovieSceneSequence* Sequence, const UClass* InClass, FGuid ObjectBinding, UControlRig* InExistingControlRig, bool bIsAdditiveControlRig, TFunction<UMovieSceneControlRigParameterTrack* (FGuid)> InTrackCreationCallback)
 {
 	FSlateApplication::Get().DismissAllMenus();
 	if (!InClass || !InClass->IsChildOf(UControlRig::StaticClass()) ||
@@ -2182,7 +2248,7 @@ static UMovieSceneControlRigParameterTrack* AddControlRig(TSharedPtr<ISequencer>
 		ExistingRigTracks = GetExistingControlRigTracksForSkeletalMeshComponent(OwnerMovieScene, SkeletalMeshComponent);
 	}
 	
-	UMovieSceneControlRigParameterTrack* Track = Cast<UMovieSceneControlRigParameterTrack>(OwnerMovieScene->AddTrack(UMovieSceneControlRigParameterTrack::StaticClass(), ObjectBinding));
+	UMovieSceneControlRigParameterTrack* Track = InTrackCreationCallback ? InTrackCreationCallback(ObjectBinding) : Cast<UMovieSceneControlRigParameterTrack>(OwnerMovieScene->AddTrack(UMovieSceneControlRigParameterTrack::StaticClass(), ObjectBinding));
 	if (Track)
 	{
 		TArray<FName> ExistingObjectNames;
@@ -2298,7 +2364,7 @@ static UMovieSceneControlRigParameterTrack* AddControlRig(TSharedPtr<ISequencer>
 	}
 	return nullptr;
 }
-static UMovieSceneTrack* FindOrCreateControlRigTrack(TSharedPtr<ISequencer>& Sequencer, UWorld* World, const UClass* ControlRigClass, const FMovieSceneBindingProxy& InBinding, bool bIsLayeredControlRig)
+static UMovieSceneTrack* FindOrCreateControlRigTrack(TSharedPtr<ISequencer>& Sequencer, UWorld* World, const UClass* ControlRigClass, const FMovieSceneBindingProxy& InBinding, bool bIsLayeredControlRig, TFunction<UMovieSceneControlRigParameterTrack*(FGuid)> InTrackCreationCallback)
 {
 	UMovieScene* MovieScene = InBinding.Sequence ? InBinding.Sequence->GetMovieScene() : nullptr;
 	UMovieSceneTrack* BaseTrack = nullptr;
@@ -2319,7 +2385,7 @@ static UMovieSceneTrack* FindOrCreateControlRigTrack(TSharedPtr<ISequencer>& Seq
 				}
 			}
 
-			UMovieSceneControlRigParameterTrack* Track = AddControlRig(Sequencer, InBinding.Sequence, ControlRigClass, InBinding.BindingID, nullptr, bIsLayeredControlRig);
+			UMovieSceneControlRigParameterTrack* Track = AddControlRig(Sequencer, InBinding.Sequence, ControlRigClass, InBinding.BindingID, nullptr, bIsLayeredControlRig, InTrackCreationCallback);
 
 			if (Track)
 			{
@@ -2330,49 +2396,49 @@ static UMovieSceneTrack* FindOrCreateControlRigTrack(TSharedPtr<ISequencer>& Seq
 	return BaseTrack;
 }
 
-void FControlRigParameterTrackEditor::AddControlRig(const UClass* InClass, UObject* BoundActor, FGuid ObjectBinding, UControlRig* InExistingControlRig)
+void FControlRigParameterTrackEditor::AddControlRig(const UClass* InClass, UObject* BoundActor, FGuid ObjectBinding, UControlRig* InExistingControlRig, TFunction<UMovieSceneControlRigParameterTrack*(FGuid)> InTrackCreationCallback)
 {
 	TSharedPtr<ISequencer> Sequencer = GetSequencer();
 	if (!Sequencer)
 	{
 		return;
 	}
-	
+
 	UWorld* World = GCurrentLevelEditingViewportClient ? GCurrentLevelEditingViewportClient->GetWorld() : nullptr;
 	UMovieSceneSequence* Sequence = Sequencer->GetFocusedMovieSceneSequence();
 	FMovieSceneBindingProxy BindingProxy(ObjectBinding, Sequence);
 
-	//UControlRigSequencerEditorLibrary::FindOrCreateControlRigTrack.. in 5.5 we will redo this but for 
+	//UControlRigSequencerEditorLibrary::FindOrCreateControlRigTrack.. in 5.5 we will redo this but for
 	//5.4.4 we can't change headers so for now we just make the change here locally
-	if (UMovieSceneTrack* Track = FindOrCreateControlRigTrack(Sequencer, World, InClass, BindingProxy, bIsLayeredControlRig))
+	if (UMovieSceneTrack* Track = FindOrCreateControlRigTrack(Sequencer, World, InClass, BindingProxy, bIsLayeredControlRig, InTrackCreationCallback))
 	{
 		BindControlRig(CastChecked<UMovieSceneControlRigParameterTrack>(Track)->GetControlRig());
 	}
 }
 
-void FControlRigParameterTrackEditor::AddControlRig(const FAssetData& InAsset, UObject* BoundActor, FGuid ObjectBinding)
+void FControlRigParameterTrackEditor::AddControlRig(const FAssetData& InAsset, UObject* BoundActor, FGuid ObjectBinding, TFunction<UMovieSceneControlRigParameterTrack*(FGuid)> InTrackCreationCallback)
 {
 	if (UControlRigBlueprint* ControlRigBlueprint = Cast<UControlRigBlueprint>(InAsset.GetAsset()))
 	{
-		AddControlRig(ControlRigBlueprint->GetRigVMBlueprintGeneratedClass(), BoundActor, ObjectBinding);
+		AddControlRig(ControlRigBlueprint->GetRigVMBlueprintGeneratedClass(), BoundActor, ObjectBinding, nullptr, InTrackCreationCallback);
 	}
 	else if (UControlRigBlueprintGeneratedClass* GeneratedClass = Cast<UControlRigBlueprintGeneratedClass>(InAsset.GetAsset()))
 	{
-		AddControlRig(GeneratedClass, BoundActor, ObjectBinding);
+		AddControlRig(GeneratedClass, BoundActor, ObjectBinding, nullptr, InTrackCreationCallback);
 	}
 }
 
-void FControlRigParameterTrackEditor::AddControlRig(const TArray<FAssetData>& InAssets, UObject* BoundActor, FGuid ObjectBinding)
+void FControlRigParameterTrackEditor::AddControlRig(const TArray<FAssetData>& InAssets, UObject* BoundActor, FGuid ObjectBinding, TFunction<UMovieSceneControlRigParameterTrack*(FGuid)> InTrackCreationCallback)
 {
 	if (InAssets.Num() > 0)
 	{
-		AddControlRig(InAssets[0], BoundActor, ObjectBinding);	
+		AddControlRig(InAssets[0], BoundActor, ObjectBinding, InTrackCreationCallback);
 	}
 }
 
-void FControlRigParameterTrackEditor::AddFKControlRig(UObject* BoundActor, FGuid ObjectBinding)
+void FControlRigParameterTrackEditor::AddFKControlRig(UObject* BoundActor, FGuid ObjectBinding, TFunction<UMovieSceneControlRigParameterTrack*(FGuid)> InTrackCreationCallback)
 {
-	AddControlRig(UFKControlRig::StaticClass(), BoundActor, ObjectBinding);
+	AddControlRig(UFKControlRig::StaticClass(), BoundActor, ObjectBinding, nullptr, InTrackCreationCallback);
 }
 
 void FControlRigParameterTrackEditor::AddControlRig(UClass* InClass, UObject* BoundActor, FGuid ObjectBinding)
@@ -3185,9 +3251,9 @@ void FControlRigParameterTrackEditor::OnSelectionChanged(TArray<UMovieSceneTrack
 
 			for (FViewModelPtr ViewModel : Sequencer->GetViewModel()->GetSelection()->Outliner)
 			{
-				if (TViewModelPtr<FTrackRowModel> TrackRowModel = ViewModel.ImplicitCast())
+				if (TViewModelPtr<ISectionOwnerExtension> SectionOwnerExtension = ViewModel.ImplicitCast())
 				{
-					for (UMovieSceneSection* Section : TrackRowModel->GetSections())
+					for (UMovieSceneSection* Section : SectionOwnerExtension->GetSections())
 					{
 						if (UMovieSceneControlRigParameterSection* CRSection = Cast<UMovieSceneControlRigParameterSection>(Section))
 						{
@@ -5443,9 +5509,9 @@ void FControlRigParameterTrackEditor::BuildTrackContextMenu(FMenuBuilder& MenuBu
 	if (bIsSection)
 	{
 		TArray<TWeakObjectPtr<UObject>> TrackSections;
-		for (UE::Sequencer::TViewModelPtr<UE::Sequencer::ITrackExtension> TrackExtension : GetSequencer()->GetViewModel()->GetSelection()->Outliner.Filter<UE::Sequencer::ITrackExtension>())
+		for (UE::Sequencer::TViewModelPtr<UE::Sequencer::ISectionOwnerExtension> SectionOwnerExtension : GetSequencer()->GetViewModel()->GetSelection()->Outliner.Filter<UE::Sequencer::ISectionOwnerExtension>())
 		{
-			for (UMovieSceneSection* Section : TrackExtension->GetSections())
+			for (UMovieSceneSection* Section : SectionOwnerExtension->GetSections())
 			{
 				TrackSections.Add(Section);
 			}
@@ -5477,16 +5543,20 @@ void FControlRigParameterTrackEditor::BuildTrackContextMenu(FMenuBuilder& MenuBu
 	{
 		MenuBuilder.BeginSection("Control Rig", LOCTEXT("ControlRig", "Control Rig"));
 		{
-			MenuBuilder.AddWidget(
-				SNew(SSpinBox<int32>)
-				.MinValue(0)
-				.Font(FAppStyle::GetFontStyle(TEXT("MenuItem.Font")))
-				.ToolTipText(LOCTEXT("OrderTooltip", "Order for this Control Rig to evaluate compared to others on the same binding, higher number means earlier evaluation"))
-				.Value_Lambda([Track]() { return Track->GetPriorityOrder(); })
-				.OnValueChanged_Lambda([Track](int32 InValue) { Track->SetPriorityOrder(InValue); })
-				,
-				LOCTEXT("Order", "Order")
-			);
+			// Don't show priority order for child tracks, it's handled by the layers in the anim mixer
+			if (Track->ParentTrack == nullptr)
+			{
+				MenuBuilder.AddWidget(
+					SNew(SSpinBox<int32>)
+					.MinValue(0)
+					.Font(FAppStyle::GetFontStyle(TEXT("MenuItem.Font")))
+					.ToolTipText(LOCTEXT("OrderTooltip", "Order for this Control Rig to evaluate compared to others on the same binding, higher number means earlier evaluation"))
+					.Value_Lambda([Track]() { return Track->GetPriorityOrder(); })
+					.OnValueChanged_Lambda([Track](int32 InValue) { Track->SetPriorityOrder(InValue); })
+					,
+					LOCTEXT("Order", "Order")
+				);
+			}
 
 			if (CVarEnableAdditiveControlRigs->GetBool())
 			{
diff --git a/Engine/Plugins/Animation/ControlRig/Source/ControlRigEditor/Private/Sequencer/ControlRigParameterTrackEditor.h b/Engine/Plugins/Animation/ControlRig/Source/ControlRigEditor/Private/Sequencer/ControlRigParameterTrackEditor.h
--- a/Engine/Plugins/Animation/ControlRig/Source/ControlRigEditor/Private/Sequencer/ControlRigParameterTrackEditor.h
+++ b/Engine/Plugins/Animation/ControlRig/Source/ControlRigEditor/Private/Sequencer/ControlRigParameterTrackEditor.h
@@ -21,6 +21,7 @@
 #include "EditorUndoClient.h"
 #include "Filters/CurveEditorSmartReduceFilter.h"
 #include "ScopedTransaction.h"
+#include "IMovieSceneAnimMixerItemMenuProvider.h"
 
 class UTickableTransformConstraint;
 struct FAssetData;
@@ -89,7 +90,7 @@ private:
  * Tools for animation tracks
  */
 class FControlRigParameterTrackEditor : public FKeyframeTrackEditor<UMovieSceneControlRigParameterTrack>, public IMovieSceneToolsAnimationBakeHelper,
-	public FEditorUndoClient
+	public FEditorUndoClient, public IMovieSceneAnimMixerItemMenuProvider
 {
 public:
 	/**
@@ -147,6 +148,12 @@ public:
 	virtual void PostUndo(bool bSuccess) override;
 	virtual void PostRedo(bool bSuccess) override { PostUndo(bSuccess); }
 
+	//IMovieSceneAnimMixerItemMenuProvider Interface
+	virtual const UClass* GetHandledMixerItemClass() const override;
+	virtual void PopulateAddMixerItemMenu(FMenuBuilder& MenuBuilder, TArray<FGuid> ObjectBindings, UMovieSceneTrack* Track, TSharedPtr<ISequencer> Sequencer, int32 RowIndex) override;
+	virtual int32 GetMixerItemMenuPriority() const override { return 50; } // Medium-high priority - appears after skeletal animation
+	virtual void PopulateObjectBindingAnimationMenu(FMenuBuilder& MenuBuilder, const TArray<FGuid>& ObjectBindings, const UClass* ObjectClass, bool bIsInsideSubmenu) override;
+
 	//function to set keys on controls
 	void SetKeysOnControls(const TMap<UControlRig*, TArray<FRigElementKey>>& Controls, EMovieSceneTransformChannel Channel);
 private:
@@ -154,7 +161,7 @@ private:
 	void ProcessKeyOperation(UObject* ObjectToKey, UMovieSceneControlRigParameterTrack* Track, TArrayView<const UE::Sequencer::FKeySectionOperation> SectionsToKey, ISequencer& InSequencer, FFrameNumber KeyTime, TArray<UE::Sequencer::FAddKeyResult>* OutResults);
 	UE::MovieScene::FControlRigParameterValues GetRecomposedControlValues(UMovieSceneControlRigParameterTrack* Track, UMovieSceneControlRigParameterSection* Section, FFrameNumber KeyTime);
 
-	void HandleAddTrackSubMenu(FMenuBuilder& MenuBuilder, TArray<FGuid> ObjectBindings, UMovieSceneTrack* Track);
+	void HandleAddTrackSubMenu(FMenuBuilder& MenuBuilder, TArray<FGuid> ObjectBindings, UMovieSceneTrack* Track, TFunction<UMovieSceneControlRigParameterTrack*(FGuid)> InTrackCreationCallback = nullptr);
 
 	void ToggleIsAdditiveControlRig();
 	bool IsToggleIsAdditiveControlRig();
@@ -168,10 +175,10 @@ private:
 
 	/** Control Rig Picked */
 	void AddControlRig(UClass* InClass, UObject* BoundActor, FGuid ObjectBinding);
-	void AddControlRig(const UClass* InClass, UObject* BoundActor, FGuid ObjectBinding, UControlRig* InExistingControlRig);
-	void AddControlRig(const FAssetData& InAsset, UObject* BoundActor, FGuid ObjectBinding);
-	void AddControlRig(const TArray<FAssetData>& InAssets, UObject* BoundActor, FGuid ObjectBinding);
-	void AddFKControlRig(UObject* BoundActor, FGuid ObjectBinding);
+	void AddControlRig(const UClass* InClass, UObject* BoundActor, FGuid ObjectBinding, UControlRig* InExistingControlRig, TFunction<UMovieSceneControlRigParameterTrack*(FGuid)> InTrackCreationCallback = nullptr);
+	void AddControlRig(const FAssetData& InAsset, UObject* BoundActor, FGuid ObjectBinding, TFunction<UMovieSceneControlRigParameterTrack*(FGuid)> InTrackCreationCallback = nullptr);
+	void AddControlRig(const TArray<FAssetData>& InAssets, UObject* BoundActor, FGuid ObjectBinding, TFunction<UMovieSceneControlRigParameterTrack*(FGuid)> InTrackCreationCallback = nullptr);
+	void AddFKControlRig(UObject* BoundActor, FGuid ObjectBinding, TFunction<UMovieSceneControlRigParameterTrack*(FGuid)> InTrackCreationCallback = nullptr);
 	void AddControlRigFromComponent(FGuid InGuid);
 
 	bool IsControlRigAllowed(const FAssetData& AssetData, TArray<UClass*> ExistingRigs, USkeleton* Skeleton);
diff --git a/Engine/Plugins/Experimental/MovieSceneAnimMixer/Source/MovieSceneAnimMixer/MovieSceneAnimMixer.Build.cs b/Engine/Plugins/Experimental/MovieSceneAnimMixer/Source/MovieSceneAnimMixer/MovieSceneAnimMixer.Build.cs
--- a/Engine/Plugins/Experimental/MovieSceneAnimMixer/Source/MovieSceneAnimMixer/MovieSceneAnimMixer.Build.cs
+++ b/Engine/Plugins/Experimental/MovieSceneAnimMixer/Source/MovieSceneAnimMixer/MovieSceneAnimMixer.Build.cs
@@ -18,7 +18,8 @@ namespace UnrealBuildTool.Rules
 					"MovieSceneTracks",
 					"AnimGraphRuntime",
 					"UAF",
-					"UAFAnimGraph"
+					"UAFAnimGraph",
+					"SlateCore"
 				}
 			);
 
diff --git a/Engine/Plugins/Experimental/MovieSceneAnimMixer/Source/MovieSceneAnimMixer/Private/AnimMixerComponentTypes.cpp b/Engine/Plugins/Experimental/MovieSceneAnimMixer/Source/MovieSceneAnimMixer/Private/AnimMixerComponentTypes.cpp
--- a/Engine/Plugins/Experimental/MovieSceneAnimMixer/Source/MovieSceneAnimMixer/Private/AnimMixerComponentTypes.cpp
+++ b/Engine/Plugins/Experimental/MovieSceneAnimMixer/Source/MovieSceneAnimMixer/Private/AnimMixerComponentTypes.cpp
@@ -47,6 +47,7 @@ namespace UE::MovieScene
 		ComponentRegistry->NewComponentType(&MixerRootMotion, TEXT("Root Motion"));
 		ComponentRegistry->NewComponentType(&MixerEntry, TEXT("MixerEntry"));
 		ComponentRegistry->NewComponentType(&RootDestination, TEXT("Root Destination"));
+		ComponentRegistry->NewComponentType(&MixerLayer, TEXT("Mixer Layer"));
 
 
 		Tags.RequiresBlending = ComponentRegistry->NewTag(TEXT("Requires Blending"));
@@ -59,6 +60,7 @@ namespace UE::MovieScene
 		ComponentRegistry->Factories.DuplicateChildComponent(MeshComponent);
 		ComponentRegistry->Factories.DuplicateChildComponent(MixerRootMotion);
 		ComponentRegistry->Factories.DuplicateChildComponent(RootDestination);
+		ComponentRegistry->Factories.DuplicateChildComponent(MixerLayer);
 
 		ComponentRegistry->Factories.DefineChildComponent(RootDestination, MixerRootMotion);
 		ComponentRegistry->Factories.DefineChildComponent(Task, MixerEntry);
diff --git a/Engine/Plugins/Experimental/MovieSceneAnimMixer/Source/MovieSceneAnimMixer/Private/MovieSceneAnimationMixerTrack.cpp b/Engine/Plugins/Experimental/MovieSceneAnimMixer/Source/MovieSceneAnimMixer/Private/MovieSceneAnimationMixerTrack.cpp
--- a/Engine/Plugins/Experimental/MovieSceneAnimMixer/Source/MovieSceneAnimMixer/Private/MovieSceneAnimationMixerTrack.cpp
+++ b/Engine/Plugins/Experimental/MovieSceneAnimMixer/Source/MovieSceneAnimMixer/Private/MovieSceneAnimationMixerTrack.cpp
@@ -9,6 +9,11 @@
 #include "AnimMixerComponentTypes.h"
 #include "Misc/AxisDisplayInfo.h"
 #include "Systems/MovieSceneRootMotionSystem.h"
+#include "MovieSceneAnimationTrackDecoration.h"
+#include "MovieSceneAnimationMixerLayer.h"
+#include "MovieScene.h"
+#include "MovieSceneSequence.h"
+#include "MovieSceneTrack.h"
 
 #define LOCTEXT_NAMESPACE "MovieSceneAnimationMixerTrack"
 
@@ -186,25 +191,25 @@ void UMovieSceneAnimationSectionDecoration::ExtendEntityImpl(UMovieSceneEntitySy
 	using namespace UE::MovieScene;
 	FAnimMixerComponentTypes* AnimMixerComponents = FAnimMixerComponentTypes::Get();
 
-	UMovieSceneAnimationMixerTrack* AnimationMixerTrack = GetTypedOuter<UMovieSceneAnimationMixerTrack>();
+	UMovieSceneSection* Section = GetTypedOuter<UMovieSceneSection>();
+	UMovieSceneAnimationMixerLayer* MixerLayer = MixerTrack ? MixerTrack->GetLayerForSection(Section) : nullptr;
 
 	TInstancedStruct<FMovieSceneMixedAnimationTarget> Target;
-	if (AnimationMixerTrack != nullptr && AnimationMixerTrack->MixedAnimationTarget.IsValid())
+	if (MixerTrack != nullptr && MixerTrack->MixedAnimationTarget.IsValid())
 	{
-		Target = AnimationMixerTrack->MixedAnimationTarget;
+		Target = MixerTrack->MixedAnimationTarget;
 	}
 	else
 	{
 		Target = TInstancedStruct<FMovieSceneMixedAnimationTarget>::Make();
 	}
 
-	UMovieSceneSection* SectionOuter = GetTypedOuter<UMovieSceneSection>();
-
 	OutImportedEntity->AddBuilder(
 		FEntityBuilder()
 		.AddTag(FMovieSceneTracksComponentTypes::Get()->Tags.AnimMixerPoseProducer)
-		.Add(AnimMixerComponents->Target, Target)		
-		.Add(AnimMixerComponents->Priority, SectionOuter ? SectionOuter->GetRowIndex() : GetRowSortOrder())
+		.Add(AnimMixerComponents->Target, Target)
+		.Add(AnimMixerComponents->Priority, MixerLayer ? MixerLayer->GetLayerIndex() : 0)
+		.Add(AnimMixerComponents->MixerLayer, MixerLayer)
 		.AddMutualComponents()
 	);
 }
@@ -218,110 +223,300 @@ UMovieSceneAnimationMixerTrack::UMovieSceneAnimationMixerTrack(const FObjectInit
 	TrackTint = FColor(66, 56, 88, 255);
 	bSupportsDefaultSections = false;
 #endif
+
+	MixedAnimationTarget = TInstancedStruct<FMovieSceneMixedAnimationTarget>::Make();
 }
 
 
 bool UMovieSceneAnimationMixerTrack::SupportsType(TSubclassOf<UMovieSceneSection> SectionClass) const
 {
-	return SectionClass.Get()->ImplementsInterface(IMovieSceneAnimationSectionInterface::UClassType::StaticClass())
-		|| SectionClass.Get() == UMovieSceneSkeletalAnimationSection::StaticClass();
+	return SectionClass.Get()->ImplementsInterface(UMovieSceneAnimationMixerItemInterface::StaticClass());
 }
 
 bool UMovieSceneAnimationMixerTrack::FixRowIndices()
 {
-	// Group sections by their type on different rows
-	TSortedMap<UClass*, TArray<UMovieSceneSection*>> SectionsByType;
-	for (UMovieSceneSection* Section : AnimationSections)
+	// Do nothing with row indices here- our rows are layers and we can have empty ones.
+	return false;
+}
+
+void UMovieSceneAnimationMixerTrack::OnSectionAddedImpl(UMovieSceneSection* Section)
+{
+	if (Section)
 	{
-		SectionsByType.FindOrAdd(Section->GetClass()).Add(Section);
-	}
+		UMovieSceneAnimationSectionDecoration* NewDecoration = NewObject<UMovieSceneAnimationSectionDecoration>(Section, UMovieSceneAnimationSectionDecoration::StaticClass(), NAME_None, RF_Transactional);
+		NewDecoration->MixerTrack = this;
+		Section->AddDecoration(NewDecoration);
 
-	bool bMadeChanges = false;
-	int32 CurrentRowIndex = 0;
+		if (UMovieSceneSkeletalAnimationSection* AnimSection = Cast<UMovieSceneSkeletalAnimationSection>(Section))
+		{
+			UMovieSceneAnimationBaseTransformDecoration* BaseTransform = NewObject<UMovieSceneAnimationBaseTransformDecoration>(AnimSection, UMovieSceneAnimationBaseTransformDecoration::StaticClass(), NAME_None, RF_Transactional);
+			AnimSection->AddDecoration(BaseTransform);
+		}
 
-	TArray<UClass*> SortedSectionTypes;
-	SectionsByType.GenerateKeyArray(SortedSectionTypes);
+#if WITH_EDITOR
+		if (IMovieSceneAnimationMixerItemInterface* AnimationItemInterface = Cast<IMovieSceneAnimationMixerItemInterface>(Section))
+		{
+			Section->SetColorTint(AnimationItemInterface->GetMixerItemTint());
+		}
+#endif
 
-	Algo::SortBy(SortedSectionTypes, [](UClass* Class){
-		IMovieSceneAnimationSectionInterface* Interface = Cast<IMovieSceneAnimationSectionInterface>(Class->GetDefaultObject());
-		return Interface ? Interface->GetRowSortOrder() : 0;
-	});
+		// Create/update layer for this section
+		int32 RowIndex = Section->GetRowIndex();
+		UMovieSceneAnimationMixerLayer* Layer = GetOrCreateLayer(RowIndex);
+		if (Layer)
+		{
+			Layer->AddSection(Section);
+		}
+	}
+}
 
-	for (UClass* Key : SortedSectionTypes)
+void UMovieSceneAnimationMixerTrack::OnSectionRemovedImpl(UMovieSceneSection* Section)
+{
+	// Remove section from its layer
+	if (Section)
 	{
-		TArray<UMovieSceneSection*>& SectionsOfType = SectionsByType.FindChecked(Key);
+		UMovieSceneAnimationMixerLayer* Layer = GetLayerForSection(Section);
+		if (Layer)
+		{
+			Layer->RemoveSection(Section);
+		}
+	}
+}
 
-		// Sort the sections
-		Algo::SortBy(SectionsOfType, &UMovieSceneSection::GetRowIndex);
 
-		int32 PreviousIndex = SectionsOfType[0]->GetRowIndex();
-		for (UMovieSceneSection* Section : SectionsOfType)
-		{
-			const int32 ThisIndex = Section->GetRowIndex();
-			if (PreviousIndex != ThisIndex)
-			{
-				++CurrentRowIndex;
-				PreviousIndex = ThisIndex;
-			}
+#if WITH_EDITORONLY_DATA
+
+FText UMovieSceneAnimationMixerTrack::GetTrackRowDisplayName(int32 RowIndex) const
+{
+	return FText::Format(LOCTEXT("AnimMixer", "Mix Layer {0}"), RowIndex);
+}
 
-			if (ThisIndex != CurrentRowIndex)
+FText UMovieSceneAnimationMixerTrack::GetDefaultDisplayName() const
+{
+	return NSLOCTEXT("AnimMixer", "DefaultTrackName", "Animation Mixer");
+}
+
+bool UMovieSceneAnimationMixerTrack::CanRename() const
+{
+	return true;
+}
+
+#endif // WITH_EDITORONLY_DATA
+
+
+UMovieSceneTrack* UMovieSceneAnimationMixerTrack::AddChildTrack(FGuid ObjectBinding, TSubclassOf<UMovieSceneTrack> TrackClass, int32 RowIndex)
+{
+	if (UMovieSceneSequence* OuterSequence = GetTypedOuter<UMovieSceneSequence>())
+	{
+		if (UMovieScene* MovieScene = OuterSequence->GetMovieScene())
+		{
+			if (ObjectBinding.IsValid())
 			{
-				bMadeChanges = true;
-				Section->SetRowIndex(CurrentRowIndex);
+				Modify();
+				UMovieSceneTrack* NewTrack = NewObject<UMovieSceneTrack>(this, TrackClass, NAME_None, RF_Transactional);
+				// Set that this is a child track. This ensures it will be not be listed with the main tracks of the sequence
+				NewTrack->ParentTrack = this;
+				if (MovieScene->AddGivenTrack(NewTrack, ObjectBinding))
+				{
+					ChildTracks.Add(NewTrack, RowIndex);
+
+					// Add track decoration to the child track- this handles adding required components to entities the track manages
+					UMovieSceneAnimationTrackDecoration* Decoration = NewObject<UMovieSceneAnimationTrackDecoration>(NewTrack, NAME_None, RF_Transactional);
+					Decoration->MixerTrack = this;
+					NewTrack->AddDecoration(Decoration);
+
+					// Create/update layer for this child track
+					UMovieSceneAnimationMixerLayer* Layer = GetOrCreateLayer(RowIndex);
+					if (Layer)
+					{
+						Layer->SetChildTrack(NewTrack);
+					}
+
+					return NewTrack;
+				}
 			}
 		}
+	}
+	return nullptr;
+}
 
-		++CurrentRowIndex;
+int32 UMovieSceneAnimationMixerTrack::GetNextChildTrackRowIndex() const
+{
+	int32 MaxRowIndex = -1;
+	for (const auto& Pair : ChildTracks)
+	{
+		MaxRowIndex = FMath::Max(MaxRowIndex, Pair.Value);
 	}
+	return MaxRowIndex + 1;
+}
 
-	return bMadeChanges;
+void UMovieSceneAnimationMixerTrack::AddSectionWithAutoRow(UMovieSceneSection* Section)
+{
+	if (!Section)
+	{
+		return;
+	}
+
+	// Rename the section to have this track as the outer
+	Section->Rename(nullptr, this);
+
+	// Compute next available row considering both sections and child tracks
+	int32 NextRowIndex = ComputeNextAvailableRowIndex();
+
+	// Assign next available row
+	Section->SetRowIndex(NextRowIndex);
+
+	// Add the section
+	AddSection(*Section);
 }
 
-void UMovieSceneAnimationMixerTrack::OnSectionAddedImpl(UMovieSceneSection* Section)
+int32 UMovieSceneAnimationMixerTrack::ComputeNextAvailableRowIndex() const
 {
-	if (UMovieSceneSkeletalAnimationSection* AnimSection = Cast<UMovieSceneSkeletalAnimationSection>(Section))
+	return Layers.Num();
+}
+
+void UMovieSceneAnimationMixerTrack::GetAllChildTracks(TArray<TObjectPtr<UMovieSceneTrack>>& OutChildTracks) const
+{
+	ChildTracks.GenerateKeyArray(OutChildTracks);
+}
+
+void UMovieSceneAnimationMixerTrack::RemoveChildTrack(UMovieSceneTrack* Track)
+{
+	if (!Track)
 	{
-		UMovieSceneAnimationSectionDecoration* NewDecoration = NewObject<UMovieSceneAnimationSectionDecoration>(AnimSection, UMovieSceneAnimationSectionDecoration::StaticClass(), NAME_None, RF_Transactional);
-		AnimSection->AddDecoration(NewDecoration);
+		return;
+	}
+
+	Modify();
+
 
-		UMovieSceneAnimationBaseTransformDecoration* BaseTransform = NewObject<UMovieSceneAnimationBaseTransformDecoration>(AnimSection, UMovieSceneAnimationBaseTransformDecoration::StaticClass(), NAME_None, RF_Transactional);
-		AnimSection->AddDecoration(BaseTransform);
+	// Remove child track reference from its layer
+	UMovieSceneAnimationMixerLayer* Layer = GetLayerForChildTrack(Track);
+	if (Layer)
+	{
+		Layer->SetChildTrack(nullptr);
+	}
+
+	ChildTracks.Remove(Track);
+
+	// Remove track decoration (if it exists)
+	Track->RemoveDecoration(UMovieSceneAnimationTrackDecoration::StaticClass());
+
+	if (UMovieSceneSequence* OuterSequence = GetTypedOuter<UMovieSceneSequence>())
+	{
+		if (UMovieScene* MovieScene = OuterSequence->GetMovieScene())
+		{
+			MovieScene->RemoveTrack(*Track);
+		}
 	}
+}
 
-	if (IMovieSceneAnimationSectionInterface* AnimationSectionInterface = Cast<IMovieSceneAnimationSectionInterface>(Section))
+UMovieSceneTrack* UMovieSceneAnimationMixerTrack::GetChildTrackAtRow(int32 RowIndex) const
+{
+	for (const auto& Pair : ChildTracks)
 	{
-		Section->SetColorTint(AnimationSectionInterface->GetMixerSectionTint());
+		if (Pair.Value == RowIndex)
+		{
+			return Pair.Key;
+		}
 	}
+	return nullptr;
 }
 
+bool UMovieSceneAnimationMixerTrack::IsChildTrack(UMovieSceneTrack* Track) const
+{
+	return ChildTracks.Contains(Track);
+}
 
-#if WITH_EDITORONLY_DATA
+int32 UMovieSceneAnimationMixerTrack::GetChildTrackRow(UMovieSceneTrack* Track) const
+{
+	const int32* RowPtr = ChildTracks.Find(Track);
 
-FText UMovieSceneAnimationMixerTrack::GetTrackRowDisplayName(int32 RowIndex) const
+	return RowPtr ? *RowPtr : INDEX_NONE;
+}
+
+void UMovieSceneAnimationMixerTrack::SetChildTrackRow(UMovieSceneTrack* Track, int32 NewRowIndex)
+{
+	if (!Track || !IsChildTrack(Track))
+	{
+		return;
+	}
+
+	int32 OldRowIndex = GetChildTrackRow(Track);
+	if (OldRowIndex == NewRowIndex)
+	{
+		return;
+	}
+
+	Modify();
+
+	ChildTracks[Track] = NewRowIndex;
+}
+
+UMovieSceneAnimationMixerLayer* UMovieSceneAnimationMixerTrack::GetOrCreateLayer(int32 RowIndex)
 {
-	// Row display name is defined by the first section in that row
-	for (UMovieSceneSection* Section : AnimationSections)
+	// Find existing layer at this row index
+	for (UMovieSceneAnimationMixerLayer* Layer : Layers)
 	{
-		if (Section->GetRowIndex() == RowIndex)
+		if (Layer)
 		{
-			return Section->GetClass()->GetDisplayNameText();
+			int32 LayerIndex = Layers.IndexOfByKey(Layer);
+			if (LayerIndex == RowIndex)
+			{
+				return Layer;
+			}
 		}
 	}
 
-	return Super::GetTrackRowDisplayName(RowIndex);
+	// Create new layer
+	Modify();
+	UMovieSceneAnimationMixerLayer* NewLayer = NewObject<UMovieSceneAnimationMixerLayer>(this, NAME_None, RF_Transactional);
+
+	// Insert at correct position to keep array sorted by row index
+	if (RowIndex >= Layers.Num())
+	{
+		// Expand array if needed
+		Layers.SetNum(RowIndex + 1);
+	}
+
+	Layers[RowIndex] = NewLayer;
+	return NewLayer;
 }
 
-FText UMovieSceneAnimationMixerTrack::GetDefaultDisplayName() const
+UMovieSceneAnimationMixerLayer* UMovieSceneAnimationMixerTrack::GetLayer(int32 RowIndex) const
 {
-	return NSLOCTEXT("AnimMixer", "DefaultTrackName", "Animation Mixer");
+	if (Layers.IsValidIndex(RowIndex))
+	{
+		return Layers[RowIndex];
+	}
+	return nullptr;
 }
 
-bool UMovieSceneAnimationMixerTrack::CanRename() const
+UMovieSceneAnimationMixerLayer* UMovieSceneAnimationMixerTrack::GetLayerForSection(UMovieSceneSection* Section) const
 {
-	return true;
+	if (!Section)
+	{
+		return nullptr;
+	}
+
+	int32 RowIndex = Section->GetRowIndex();
+	return GetLayer(RowIndex);
 }
 
-#endif // WITH_EDITORONLY_DATA
+UMovieSceneAnimationMixerLayer* UMovieSceneAnimationMixerTrack::GetLayerForChildTrack(UMovieSceneTrack* Track) const
+{
+	if (!Track)
+	{
+		return nullptr;
+	}
+
+	const int32* RowIndex = ChildTracks.Find(Track);
+	if (RowIndex)
+	{
+		return GetLayer(*RowIndex);
+	}
+	return nullptr;
+}
 
 
 #undef LOCTEXT_NAMESPACE
\ No newline at end of file
diff --git a/Engine/Plugins/Experimental/MovieSceneAnimMixer/Source/MovieSceneAnimMixer/Private/Systems/MovieSceneMixedSkeletalAnimationSystem.cpp b/Engine/Plugins/Experimental/MovieSceneAnimMixer/Source/MovieSceneAnimMixer/Private/Systems/MovieSceneMixedSkeletalAnimationSystem.cpp
--- a/Engine/Plugins/Experimental/MovieSceneAnimMixer/Source/MovieSceneAnimMixer/Private/Systems/MovieSceneMixedSkeletalAnimationSystem.cpp
+++ b/Engine/Plugins/Experimental/MovieSceneAnimMixer/Source/MovieSceneAnimMixer/Private/Systems/MovieSceneMixedSkeletalAnimationSystem.cpp
@@ -120,6 +120,8 @@ struct FSkeletalAnimMixerMutation : IMovieSceneEntityMutation
 		TComponentWriter<TSharedPtr<FAnimNextEvaluationTask>> OutTasks = Allocation->WriteComponents(AnimMixerComponents->Task, WriteContext);
 		TComponentWriter<int32> OutPriorities = Allocation->WriteComponents(AnimMixerComponents->Priority, WriteContext);
 		TComponentWriter<FMovieSceneRootMotionSettings> OutRootMotionSettings = Allocation->WriteComponents(AnimMixerComponents->RootMotionSettings, WriteContext);
+		TOptionalComponentReader<TObjectPtr<UMovieSceneAnimationMixerLayer>> OptMixerLayers = Allocation->TryReadComponents(AnimMixerComponents->MixerLayer);
+		
 		for (int32 Index = 0; Index < Num; ++Index)
 		{
 			// Initialize a null (empty) task, will be written during evaluation
@@ -128,13 +130,9 @@ struct FSkeletalAnimMixerMutation : IMovieSceneEntityMutation
 			const FMovieSceneSkeletalAnimationComponentData& SkeletalAnim = SkeletalAnimationData[Index];
 			if (SkeletalAnim.Section)
 			{
-				UMovieSceneAnimationMixerTrack* MixerTrack = Cast<UMovieSceneAnimationMixerTrack>(SkeletalAnim.Section->GetOuter());
-				
-				// If we have a mixer track, the animation decoration may already have added our target and priority
-				
-				if (MixerTrack == nullptr || !MixerTrack->MixedAnimationTarget.IsValid())
+				// If we're not on a mixer layer, set a default target and priority
+				if (!OptMixerLayers || !OptMixerLayers[Index])
 				{
-					// Set a default target- Automatic will resolve to the appropriate place, or the Anim Mixer track will override it.
 					OutAnimTargets[Index] = TInstancedStruct<FMovieSceneMixedAnimationTarget>::Make();
 					OutPriorities[Index] = -100000;
 				}
diff --git a/Engine/Plugins/Experimental/MovieSceneAnimMixer/Source/MovieSceneAnimMixer/Public/AnimMixerComponentTypes.h b/Engine/Plugins/Experimental/MovieSceneAnimMixer/Source/MovieSceneAnimMixer/Public/AnimMixerComponentTypes.h
--- a/Engine/Plugins/Experimental/MovieSceneAnimMixer/Source/MovieSceneAnimMixer/Public/AnimMixerComponentTypes.h
+++ b/Engine/Plugins/Experimental/MovieSceneAnimMixer/Source/MovieSceneAnimMixer/Public/AnimMixerComponentTypes.h
@@ -16,6 +16,7 @@
 enum class EMovieSceneRootMotionDestination : uint8;
 
 struct FMovieSceneAnimMixerEntry;
+class UMovieSceneAnimationMixerLayer;
 
 namespace UE::MovieScene
 {
@@ -40,6 +41,7 @@ namespace UE::MovieScene
 		TComponentTypeID<EMovieSceneRootMotionDestination> RootDestination;
 		TComponentTypeID<FObjectComponent> MeshComponent;
 		TComponentTypeID<TSharedPtr<FMovieSceneMixerRootMotionComponentData>> MixerRootMotion;
+		TComponentTypeID<TObjectPtr<UMovieSceneAnimationMixerLayer>> MixerLayer;
 
 		struct
 		{
diff --git a/Engine/Plugins/Experimental/MovieSceneAnimMixer/Source/MovieSceneAnimMixer/Public/MovieSceneAnimationMixerTrack.h b/Engine/Plugins/Experimental/MovieSceneAnimMixer/Source/MovieSceneAnimMixer/Public/MovieSceneAnimationMixerTrack.h
--- a/Engine/Plugins/Experimental/MovieSceneAnimMixer/Source/MovieSceneAnimMixer/Public/MovieSceneAnimationMixerTrack.h
+++ b/Engine/Plugins/Experimental/MovieSceneAnimMixer/Source/MovieSceneAnimMixer/Public/MovieSceneAnimationMixerTrack.h
@@ -10,11 +10,14 @@
 #include "EntitySystem/IMovieSceneEntityDecorator.h"
 #include "MovieSceneMixedAnimationTarget.h"
 #include "MovieSceneNameableTrack.h"
+#include "MovieSceneAnimationMixerItemInterface.h"
 #include "StructUtils/InstancedStruct.h"
 #include "Tracks/MovieSceneCommonAnimationTrack.h"
 #include "Systems/MovieSceneRootMotionSystem.h"
 #include "MovieSceneAnimationMixerTrack.generated.h"
 
+class UMovieSceneAnimationMixerLayer;
+
 USTRUCT()
 struct FMovieSceneByteChannelDefaultOnly : public FMovieSceneByteChannel
 {
@@ -40,48 +43,29 @@ struct TMovieSceneChannelTraits<FMovieSceneByteChannelDefaultOnly> : TMovieScene
 #endif
 };
 
-UINTERFACE(MinimalAPI)
-class UMovieSceneAnimationSectionInterface : public UInterface
-{
-	GENERATED_BODY()
-};
-
-
-class IMovieSceneAnimationSectionInterface
-{
-public:
-	GENERATED_BODY()
-
-	virtual int32 GetRowSortOrder() const = 0;
-
-	// Allows different section types to have different colors when used in the mixer, without needing to override their colors everywhere.
-	virtual FColor GetMixerSectionTint() const = 0;
-};
-
 UCLASS(MinimalAPI)
 class UMovieSceneAnimationSectionDecoration
 	: public UMovieSceneSignedObject
 	, public IMovieSceneEntityDecorator
 	, public IMovieSceneSectionDecoration
-	, public IMovieSceneAnimationSectionInterface
+	, public IMovieSceneAnimationMixerItemInterface
 {
 	GENERATED_BODY()
 
-	virtual void ExtendEntityImpl(UMovieSceneEntitySystemLinker* EntityLinker, const FEntityImportParams& Params, FImportedEntity* OutImportedEntity) override;
+public:
 
-	virtual int32 GetRowSortOrder() const
-	{
-		return RowSortOrder;
-	}
+	UPROPERTY()
+	TObjectPtr<UMovieSceneAnimationMixerTrack> MixerTrack;
+
+protected: 
 
-	virtual FColor GetMixerSectionTint() const override
+	virtual void ExtendEntityImpl(UMovieSceneEntitySystemLinker* EntityLinker, const FEntityImportParams& Params, FImportedEntity* OutImportedEntity) override;
+
+	virtual FColor GetMixerItemTint() const override
 	{
 		return MixerTintOverride;
 	}
 
-	UPROPERTY()
-	int32 RowSortOrder = 0;
-
 private:
 	FColor MixerTintOverride = FColor();
 };
@@ -137,6 +121,10 @@ public:
 	virtual bool SupportsType(TSubclassOf<UMovieSceneSection> SectionClass) const override;
 
 	virtual void OnSectionAddedImpl(UMovieSceneSection* Secton) override;
+	virtual void OnSectionRemovedImpl(UMovieSceneSection* Section) override;
+
+	MOVIESCENEANIMMIXER_API virtual void AddSectionWithAutoRow(UMovieSceneSection* Section) override;
+	MOVIESCENEANIMMIXER_API virtual int32 ComputeNextAvailableRowIndex() const override;
 
 #if WITH_EDITORONLY_DATA
 
@@ -145,8 +133,48 @@ public:
 	virtual bool CanRename() const override;
 #endif
 
-public:
+	MOVIESCENEANIMMIXER_API virtual UMovieSceneTrack* AddChildTrack(FGuid ObjectBinding, TSubclassOf<UMovieSceneTrack> TrackClass, int32 RowIndex) override;
+	MOVIESCENEANIMMIXER_API int32 GetNextChildTrackRowIndex() const;
+	MOVIESCENEANIMMIXER_API virtual void GetAllChildTracks(TArray<TObjectPtr<UMovieSceneTrack>>& OutChildTracks) const override;
+	MOVIESCENEANIMMIXER_API void RemoveChildTrack(UMovieSceneTrack* Track);
+
+	MOVIESCENEANIMMIXER_API UMovieSceneTrack* GetChildTrackAtRow(int32 RowIndex) const;
+
+	MOVIESCENEANIMMIXER_API bool IsChildTrack(UMovieSceneTrack* Track) const;
+	MOVIESCENEANIMMIXER_API int32 GetChildTrackRow(UMovieSceneTrack* Track) const;
+	MOVIESCENEANIMMIXER_API void SetChildTrackRow(UMovieSceneTrack* Track, int32 NewRowIndex);
+
+	// Layer Management
+	// Note: Layer index is implicitly the row index. Sections use Section->GetRowIndex(),
+	// Child tracks are referenced by individual layers, and cached in ChildTracks map. 
+	// Layers array is kept sorted by row index.
+
+	/** Create or get a layer at the specified row index */
+	MOVIESCENEANIMMIXER_API UMovieSceneAnimationMixerLayer* GetOrCreateLayer(int32 RowIndex);
 
+	/** Get the layer at the specified row index (returns null if doesn't exist) */
+	MOVIESCENEANIMMIXER_API UMovieSceneAnimationMixerLayer* GetLayer(int32 RowIndex) const;
+
+	/** Get all layers (sorted by row index) */
+	MOVIESCENEANIMMIXER_API const TArray<TObjectPtr<UMovieSceneAnimationMixerLayer>>& GetLayers() const { return Layers; }
+
+	/** Get the layer containing the given section (uses section's row index) */
+	MOVIESCENEANIMMIXER_API UMovieSceneAnimationMixerLayer* GetLayerForSection(UMovieSceneSection* Section) const;
+
+	/** Get the layer containing the given child track */
+	MOVIESCENEANIMMIXER_API UMovieSceneAnimationMixerLayer* GetLayerForChildTrack(UMovieSceneTrack* Track) const;
+
+public: 
 	UPROPERTY(EditAnywhere, Category = "Mixed Animation")
 	TInstancedStruct<FMovieSceneMixedAnimationTarget> MixedAnimationTarget;
+
+private:
+	// Row assignment for child tracks
+	UPROPERTY()
+	TMap<TObjectPtr<UMovieSceneTrack>, int32> ChildTracks;
+
+	// Layers organizing sections and child tracks by row index
+	// Array is kept sorted by row index for efficient lookup
+	UPROPERTY()
+	TArray<TObjectPtr<UMovieSceneAnimationMixerLayer>> Layers;
 };
diff --git a/Engine/Plugins/Experimental/MovieSceneAnimMixer/Source/MovieSceneAnimMixer/Public/MovieSceneRootMotionSection.h b/Engine/Plugins/Experimental/MovieSceneAnimMixer/Source/MovieSceneAnimMixer/Public/MovieSceneRootMotionSection.h
--- a/Engine/Plugins/Experimental/MovieSceneAnimMixer/Source/MovieSceneAnimMixer/Public/MovieSceneRootMotionSection.h
+++ b/Engine/Plugins/Experimental/MovieSceneAnimMixer/Source/MovieSceneAnimMixer/Public/MovieSceneRootMotionSection.h
@@ -31,7 +31,7 @@ UCLASS(MinimalAPI, DisplayName="Root Motion")
 class UMovieSceneRootMotionSection
 	: public UMovieSceneSection
 	, public IMovieSceneEntityProvider
-	, public IMovieSceneAnimationSectionInterface
+	, public IMovieSceneAnimationMixerItemInterface
 {
 	GENERATED_BODY()
 
@@ -41,7 +41,7 @@ public:
 
 	virtual int32 GetRowSortOrder() const;
 
-	virtual FColor GetMixerSectionTint() const override
+	virtual FColor GetMixerItemTint() const override
 	{
 		return MixerTintOverride;
 	}
diff --git a/Engine/Plugins/Experimental/MovieSceneAnimMixer/Source/MovieSceneAnimMixerEditor/Private/MovieSceneAnimMixerEditorModule.cpp b/Engine/Plugins/Experimental/MovieSceneAnimMixer/Source/MovieSceneAnimMixerEditor/Private/MovieSceneAnimMixerEditorModule.cpp
--- a/Engine/Plugins/Experimental/MovieSceneAnimMixer/Source/MovieSceneAnimMixerEditor/Private/MovieSceneAnimMixerEditorModule.cpp
+++ b/Engine/Plugins/Experimental/MovieSceneAnimMixer/Source/MovieSceneAnimMixerEditor/Private/MovieSceneAnimMixerEditorModule.cpp
@@ -10,6 +10,7 @@
 #include "MovieSceneAnimationMixerTrack.h"
 #include "MovieSceneAnimationMixerTrackEditor.h"
 #include "Channels/BuiltInChannelEditors.h"
+#include "MVVM/AnimationMixerTrackModel.h"
 
 
 namespace Sequencer
@@ -91,6 +92,17 @@ namespace UE::MovieScene
 		ISequencerModule& SequencerModule = FModuleManager::Get().LoadModuleChecked<ISequencerModule>(SequencerModuleName.Resolve());
 		AnimationTrackEditorHandle = SequencerModule.RegisterTrackEditor(FOnCreateTrackEditor::CreateStatic( &FAnimationMixerTrackEditor::CreateTrackEditor));
 
+		// Register custom track model for Animation Mixer tracks
+		AnimationMixerTrackModelHandle = SequencerModule.RegisterTrackModel(FOnCreateTrackModel::CreateLambda([](UMovieSceneTrack* Track) -> TSharedPtr<FTrackModel>
+			{
+				if (Track && Track->IsA<UMovieSceneAnimationMixerTrack>())
+				{
+					return MakeShared<FAnimationMixerTrackModel>(Track);
+				}
+				return nullptr;
+			}
+		));
+
 		SequencerModule.RegisterChannelInterface<FMovieSceneByteChannelDefaultOnly>();
 	}
 
@@ -98,7 +110,7 @@ namespace UE::MovieScene
 	{
 		if (ISequencerModule* SequencerModule = FModuleManager::Get().GetModulePtr<ISequencerModule>(SequencerModuleName.Resolve()))
 		{
-			SequencerModule->UnregisterTrackModel(AnimationTrackEditorHandle);
+			SequencerModule->UnregisterTrackModel(AnimationMixerTrackModelHandle);
 		}
 
 		IModularFeatures::Get().UnregisterModularFeature(IAnimGraph_SequencerMixerTargetConnector::GetModularFeatureName(), this);
diff --git a/Engine/Plugins/Experimental/MovieSceneAnimMixer/Source/MovieSceneAnimMixerEditor/Private/MovieSceneAnimMixerEditorModule.h b/Engine/Plugins/Experimental/MovieSceneAnimMixer/Source/MovieSceneAnimMixerEditor/Private/MovieSceneAnimMixerEditorModule.h
--- a/Engine/Plugins/Experimental/MovieSceneAnimMixer/Source/MovieSceneAnimMixerEditor/Private/MovieSceneAnimMixerEditorModule.h
+++ b/Engine/Plugins/Experimental/MovieSceneAnimMixer/Source/MovieSceneAnimMixerEditor/Private/MovieSceneAnimMixerEditorModule.h
@@ -20,6 +20,7 @@ namespace UE::MovieScene
 
 	private:
 		FDelegateHandle AnimationTrackEditorHandle;
+		FDelegateHandle AnimationMixerTrackModelHandle;
 	};
 }
 
diff --git a/Engine/Plugins/Experimental/MovieSceneAnimMixer/Source/MovieSceneAnimMixerEditor/Private/MovieSceneAnimationMixerTrackEditor.cpp b/Engine/Plugins/Experimental/MovieSceneAnimMixer/Source/MovieSceneAnimMixerEditor/Private/MovieSceneAnimationMixerTrackEditor.cpp
--- a/Engine/Plugins/Experimental/MovieSceneAnimMixer/Source/MovieSceneAnimMixerEditor/Private/MovieSceneAnimationMixerTrackEditor.cpp
+++ b/Engine/Plugins/Experimental/MovieSceneAnimMixer/Source/MovieSceneAnimMixerEditor/Private/MovieSceneAnimationMixerTrackEditor.cpp
@@ -8,10 +8,14 @@
 #include "Components/SkeletalMeshComponent.h"
 #include "Framework/MultiBox/MultiBoxBuilder.h"
 #include "GameFramework/Actor.h"
+#include "IMovieSceneAnimMixerItemMenuProvider.h"
 #include "ISequencer.h"
+#include "MovieSceneAnimationMixerItemInterface.h"
 #include "MovieSceneAnimationMixerTrack.h"
+#include "MovieSceneTrack.h"
 #include "MVVM/Extensions/IObjectBindingExtension.h"
 #include "MVVM/Extensions/ITrackExtension.h"
+#include "MVVM/Extensions/ITrackRowExtension.h"
 #include "MVVM/ViewModels/OutlinerColumns/OutlinerColumnTypes.h"
 #include "MVVM/ViewModels/SequencerEditorViewModel.h"
 #include "MVVM/ViewModels/TrackModel.h"
@@ -20,6 +24,8 @@
 #include "ScopedTransaction.h"
 #include "Sections/MovieSceneSkeletalAnimationSection.h"
 #include "SequencerUtilities.h"
+#include "MovieSceneAnimationMixerLayer.h"
+#include "Features/IModularFeatures.h"
 
 #define LOCTEXT_NAMESPACE "MovieSceneAnimationMixerTrackEditor"
 
@@ -222,12 +228,48 @@ FAnimationMixerTrackEditor::FAnimationMixerTrackEditor(TSharedRef<ISequencer> In
 {
 }
 
+void FAnimationMixerTrackEditor::OnInitialize()
+{
+	FCommonAnimationTrackEditor::OnInitialize();
+
+	// Register as animation mixer item menu provider
+	IModularFeatures::Get().RegisterModularFeature(IMovieSceneAnimMixerItemMenuProvider::GetModularFeatureName(), this);
+}
+
+void FAnimationMixerTrackEditor::OnRelease()
+{
+	// Unregister as animation mixer item menu provider
+	IModularFeatures::Get().UnregisterModularFeature(IMovieSceneAnimMixerItemMenuProvider::GetModularFeatureName(), this);
+
+	FCommonAnimationTrackEditor::OnRelease();
+}
+
+const UClass* FAnimationMixerTrackEditor::GetHandledMixerItemClass() const
+{
+	// Animation Mixer track editor doesn't handle any specific mixer item class
+	// It only provides the "Animation Mixer" entry in the Object Binding Animation menu
+	return nullptr;
+}
+
 FText FAnimationMixerTrackEditor::GetDisplayName() const
 {
 	return LOCTEXT("AnimationMixerTrackEditor_DisplayName", "Animation Mixer");
 }
 
 void FAnimationMixerTrackEditor::BuildObjectBindingTrackMenu(FMenuBuilder& MenuBuilder, const TArray<FGuid>& ObjectBindings, const UClass* ObjectClass)
+{
+	// Animation Mixer contributes to the Animation submenu via PopulateObjectBindingAnimationMenu instead.
+	// The base class FCommonAnimationTrackEditor::BuildObjectBindingTrackMenu creates the Animation submenu,
+	// so we override to empty here to avoid creating duplicate Animation submenus.
+}
+
+void FAnimationMixerTrackEditor::PopulateAddMixerItemMenu(FMenuBuilder& MenuBuilder, TArray<FGuid> ObjectBindings, UMovieSceneTrack* Track, TSharedPtr<ISequencer> Sequencer, int32 RowIndex)
+{
+	// Animation Mixer track editor doesn't contribute skeletal animations to the mixer '+' menu.
+	// FSkeletalAnimationTrackEditor (which also inherits from FCommonAnimationTrackEditor) handles that.
+}
+
+void FAnimationMixerTrackEditor::PopulateObjectBindingAnimationMenu(FMenuBuilder& MenuBuilder, const TArray<FGuid>& ObjectBindings, const UClass* ObjectClass, bool bIsInsideSubmenu)
 {
 	if (ObjectClass != nullptr && (ObjectClass->IsChildOf(USkeletalMeshComponent::StaticClass()) || ObjectClass->IsChildOf(AActor::StaticClass())))
 	{
@@ -235,7 +277,7 @@ void FAnimationMixerTrackEditor::BuildObjectBindingTrackMenu(FMenuBuilder& MenuB
 			LOCTEXT("AddAnimationTrack", "Animation Mixer"),
 			LOCTEXT("AddAnimationTrackTooltip", "Adds a new animation track for playing back Anim Sequences and other sources of animation."),
 			FSlateIcon(FAppStyle::GetAppStyleSetName(), "Sequencer.Tracks.Animation"),
-			FUIAction( 
+			FUIAction(
 				FExecuteAction::CreateSP(this, &FAnimationMixerTrackEditor::HandleAddAnimationTrackMenuEntryExecute, ObjectBindings)
 			)
 		);
@@ -255,6 +297,22 @@ TSharedPtr<SWidget> FAnimationMixerTrackEditor::BuildOutlinerColumnWidget(const
 
 	if (ColumnName == FCommonOutlinerNames::Add)
 	{
+		// Don't show '+' button on layers with child tracks
+		if (TViewModelPtr<ITrackRowExtension> TrackRowExtension = Params.ViewModel.ImplicitCast())
+		{
+			UMovieSceneAnimationMixerTrack* MixerTrack = Cast<UMovieSceneAnimationMixerTrack>(TrackRowExtension->GetParentTrack());
+			if (MixerTrack)
+			{
+				const int32 LayerRowIndex = TrackRowExtension->GetRowIndex();
+				UMovieSceneAnimationMixerLayer* Layer = MixerTrack->GetLayer(LayerRowIndex);
+				if (Layer && Layer->HasChildTrack())
+				{
+					// Layer has a child track - don't show '+' button
+					return SNullWidget::NullWidget;
+				}
+			}
+		}
+
 		return MakeAddButton(
 			LOCTEXT("AddSection", "Section"),
 			FOnGetContent::CreateSP(
@@ -312,11 +370,16 @@ TSharedRef<ISequencerSection> FAnimationMixerTrackEditor::MakeSectionInterface(U
 TSharedRef<SWidget> FAnimationMixerTrackEditor::BuildAddSectionSubMenu(TWeakViewModelPtr<IOutlinerExtension> WeakViewModel, TWeakViewModelPtr<FSequencerEditorViewModel> WeakEditor)
 {
 	TViewModelPtr<ITrackExtension> Track = WeakViewModel.ImplicitPin();
-	if (!Track)
+	TViewModelPtr<ITrackRowExtension> LayerExtension = WeakViewModel.ImplicitPin();
+	UMovieSceneAnimationMixerTrack* MixerTrack = LayerExtension ? Cast<UMovieSceneAnimationMixerTrack>(LayerExtension->GetParentTrack())
+													   : Track ? Cast<UMovieSceneAnimationMixerTrack>(Track->GetTrack()) : nullptr;
+	if (!MixerTrack)
 	{
 		return SNullWidget::NullWidget;
 	}
 
+
+	// Lambda for creating sections with defaults
 	auto CreateNewSection = [WeakViewModel, WeakEditor](FTopLevelAssetPath ClassPath)
 	{
 		UClass* SectionClass = FSoftClassPath(ClassPath.ToString()).TryLoadClass<UMovieSceneSection>();
@@ -325,114 +388,322 @@ TSharedRef<SWidget> FAnimationMixerTrackEditor::BuildAddSectionSubMenu(TWeakView
 			return;
 		}
 
-		check(SectionClass && SectionClass->IsChildOf(UMovieSceneSection::StaticClass()) && SectionClass->ImplementsInterface(IMovieSceneAnimationSectionInterface::UClassType::StaticClass()));
+		check(SectionClass && SectionClass->IsChildOf(UMovieSceneSection::StaticClass()) && SectionClass->ImplementsInterface(UMovieSceneAnimationMixerItemInterface::StaticClass()));
 
-		TViewModelPtr<ITrackExtension>           Track     = WeakViewModel.ImplicitPin();
+		TViewModelPtr<ITrackExtension>           Track = WeakViewModel.ImplicitPin();
+		TViewModelPtr<ITrackRowExtension> LayerExtension = WeakViewModel.ImplicitPin();
 		TViewModelPtr<FSequencerEditorViewModel> Editor    = WeakEditor.Pin();
 		TSharedPtr<ISequencer>                   Sequencer = Editor ? Editor->GetSequencer() : nullptr;
 
-		UMovieSceneTrack* TrackObject = Track ? Track->GetTrack() : nullptr;
-		if (Sequencer && TrackObject)
+		UMovieSceneAnimationMixerTrack* MixerTrack = LayerExtension ? Cast<UMovieSceneAnimationMixerTrack>(LayerExtension->GetParentTrack())
+														   : Track ? Cast<UMovieSceneAnimationMixerTrack>(Track->GetTrack()) : nullptr;
+		if (Sequencer && MixerTrack)
 		{
 			FScopedTransaction Transaction(FText::Format(LOCTEXT("AddSectionTransaction", "Add New {0} Section"), SectionClass->GetDisplayNameText()));
 
-			TrackObject->Modify();
+			MixerTrack->Modify();
 
-			UMovieSceneSection* NewSection = NewObject<UMovieSceneSection>(TrackObject, SectionClass, NAME_None, RF_Transactional);
+			UMovieSceneSection* NewSection = NewObject<UMovieSceneSection>(MixerTrack, SectionClass, NAME_None, RF_Transactional);
 
-			int32 MaxRowIndex = -1;
-			for (UMovieSceneSection* Section : TrackObject->GetAllSections())
+			// Determine row index based on whether we're adding to a layer model or the track model
+			int32 NewRowIndex;
+			if (LayerExtension)
+			{
+				// Adding to a specific layer - use the layer's row index
+				NewRowIndex = LayerExtension->GetRowIndex();
+			}
+			else
 			{
-				if (Section->GetClass() == SectionClass)
+				// Adding to the track model - find next available row
+				int32 MaxRowIndex = -1;
+				for (UMovieSceneSection* Section : MixerTrack->GetAllSections())
 				{
-					if (!Section->GetRange().Overlaps(NewSection->GetRange()))
-					{
-						if (MaxRowIndex > Section->GetRowIndex() || MaxRowIndex == -1)
-						{
-							MaxRowIndex = Section->GetRowIndex();
-						}
-					}
-					else if (MaxRowIndex == -1 || MaxRowIndex > Section->GetRowIndex()+1)
-					{
-						MaxRowIndex = Section->GetRowIndex() + 1;
-					}
+					MaxRowIndex = FMath::Max(MaxRowIndex, Section->GetRowIndex());
 				}
+
+				NewRowIndex = FMath::Max(MaxRowIndex + 1, MixerTrack->GetNextChildTrackRowIndex());
 			}
 
-			NewSection->SetRowIndex(MaxRowIndex);
+			NewSection->SetRowIndex(NewRowIndex);
 
-			TrackObject->AddSection(*NewSection);
+			MixerTrack->AddSection(*NewSection);
 
-			TrackObject->FixRowIndices();
-			TrackObject->UpdateEasing();
+			MixerTrack->UpdateEasing();
 
 			Sequencer->NotifyMovieSceneDataChanged(EMovieSceneDataChangeType::MovieSceneStructureItemAdded);
 		}
 	};
 
-	TViewModelSubListIterator<FSectionModel> Sections = Track->GetSectionModels().IterateSubList<FSectionModel>();
+	// Lambda for creating child tracks with defaults
+	auto CreateNewChildTrack = [WeakViewModel, WeakEditor](FGuid ObjectBinding, FTopLevelAssetPath ClassPath)
+	{
+		UClass* TrackClass = FSoftClassPath(ClassPath.ToString()).TryLoadClass<UMovieSceneTrack>();
+		if (!TrackClass)
+		{
+			return;
+		}
+
+		check(TrackClass && TrackClass->IsChildOf(UMovieSceneTrack::StaticClass()) && TrackClass->ImplementsInterface(UMovieSceneAnimationMixerItemInterface::StaticClass()));
+
+		TViewModelPtr<ITrackExtension>           Track = WeakViewModel.ImplicitPin(); 
+		TViewModelPtr<ITrackRowExtension> LayerExtension = WeakViewModel.ImplicitPin();
+
+		TViewModelPtr<FSequencerEditorViewModel> Editor    = WeakEditor.Pin();
+		TSharedPtr<ISequencer>                   Sequencer = Editor ? Editor->GetSequencer() : nullptr;
+
+		UMovieSceneAnimationMixerTrack* MixerTrack = LayerExtension ? Cast<UMovieSceneAnimationMixerTrack>(LayerExtension->GetParentTrack())
+														   : Track ? Cast<UMovieSceneAnimationMixerTrack>(Track->GetTrack()) : nullptr;
+		if (Sequencer && MixerTrack)
+		{
+			// Cannot add child tracks to a non-empty layer
+			if (LayerExtension)
+			{
+				const int32 LayerRowIndex = LayerExtension->GetRowIndex();
+				UMovieSceneAnimationMixerLayer* Layer = MixerTrack->GetLayer(LayerRowIndex);
+				if (Layer && !Layer->IsEmpty())
+				{
+					// Layer already has content - cannot add a child track
+					return;
+				}
+			}
+
+			FScopedTransaction Transaction(FText::Format(LOCTEXT("AddChildTrackTransaction", "Add New {0} Child Track"), TrackClass->GetDisplayNameText()));
+
+			MixerTrack->Modify();
+
+			// Determine row index based on whether we're adding to a layer model or the track model
+			int32 NewRowIndex;
+			if (LayerExtension)
+			{
+				// Adding to a specific layer - use the layer's row index
+				NewRowIndex = LayerExtension->GetRowIndex();
+			}
+			else
+			{
+				// Adding to the track model - find next available row
+				int32 MaxRowIndex = -1;
+				for (UMovieSceneSection* Section : MixerTrack->GetAllSections())
+				{
+					MaxRowIndex = FMath::Max(MaxRowIndex, Section->GetRowIndex());
+				}
+
+				NewRowIndex = FMath::Max(MaxRowIndex + 1, MixerTrack->GetNextChildTrackRowIndex());
+			}
+
+			UMovieSceneTrack* NewChildTrack = MixerTrack->AddChildTrack(ObjectBinding, TrackClass, NewRowIndex);
+
+			Sequencer->NotifyMovieSceneDataChanged(EMovieSceneDataChangeType::MovieSceneStructureItemAdded);
+		}
+	};
 
 	FMenuBuilder MenuBuilder(true, nullptr);
 
-	if (Sections && Cast<UMovieSceneSkeletalAnimationSection>(Sections->GetSection()))
+	// Layer management
+	// Only show on the track-level menu (not on individual layer menus)
+	if (!LayerExtension.IsValid())
 	{
-		MenuBuilder.BeginSection(NAME_None, LOCTEXT("TimeWarp_Label", "Time Warp"));
+		MenuBuilder.BeginSection(NAME_None, LOCTEXT("LayerManagement_Label", "Layer Management"));
 		{
-			FSequencerUtilities::MakeTimeWarpMenuEntry(MenuBuilder, Track);
+			MenuBuilder.AddMenuEntry(
+				LOCTEXT("AddLayer", "Layer"),
+				LOCTEXT("AddLayer_Tooltip", "Add an empty layer at the end of the track"),
+				FSlateIcon(),
+				FUIAction(FExecuteAction::CreateLambda([WeakViewModel, WeakEditor]()
+				{
+					TViewModelPtr<ITrackExtension>           Track     = WeakViewModel.ImplicitPin();
+					TViewModelPtr<FSequencerEditorViewModel> Editor    = WeakEditor.Pin();
+					TSharedPtr<ISequencer>                   Sequencer = Editor ? Editor->GetSequencer() : nullptr;
+
+					UMovieSceneAnimationMixerTrack* MixerTrack = Track ? Cast<UMovieSceneAnimationMixerTrack>(Track->GetTrack()) : nullptr;
+					if (Sequencer && MixerTrack)
+					{
+						FScopedTransaction Transaction(LOCTEXT("AddLayerTransaction", "Add Layer"));
+						MixerTrack->Modify();
+
+						// Compute next available row index
+						int32 NextRowIndex = MixerTrack->ComputeNextAvailableRowIndex();
+
+						// Create an empty layer at the next available row
+						MixerTrack->GetOrCreateLayer(NextRowIndex);
+
+						Sequencer->NotifyMovieSceneDataChanged(EMovieSceneDataChangeType::MovieSceneStructureItemAdded);
+					}
+				}))
+			);
 		}
 		MenuBuilder.EndSection();
 	}
 
-	// Find all classes that match our interface
+	if (TViewModelPtr<ISectionOwnerExtension> SectionOwner = WeakViewModel.ImplicitPin())
+	{ 
+		TViewModelSubListIterator<FSectionModel> Sections = SectionOwner->GetSectionModels().IterateSubList<FSectionModel>();
+
+		// Time warp entry (if applicable)
+		if (Sections && Cast<UMovieSceneSkeletalAnimationSection>(Sections->GetSection()))
+		{
+			MenuBuilder.BeginSection(NAME_None, LOCTEXT("TimeWarp_Label", "Time Warp"));
+			{
+				FSequencerUtilities::MakeTimeWarpMenuEntry(MenuBuilder, Track);
+			}
+			MenuBuilder.EndSection();
+		}
+	}
+
+	// Discover all item classes (sections and tracks) via Asset Registry
 	FAssetRegistryModule& AssetRegistryModule = FModuleManager::LoadModuleChecked<FAssetRegistryModule>("AssetRegistry");
 
-	TSet<FTopLevelAssetPath> AllClasses;
+	TSet<FTopLevelAssetPath> AllItemClasses;
+
 	{
-		FTopLevelAssetPath TargetClassPath(UMovieSceneSection::StaticClass());
+		// Get all sections
+		FTopLevelAssetPath SectionClassPath(UMovieSceneSection::StaticClass());
+		TSet<FTopLevelAssetPath> SectionClasses;
+		AssetRegistryModule.Get().GetDerivedClassNames({ SectionClassPath }, TSet<FTopLevelAssetPath>(), SectionClasses);
+		AllItemClasses.Append(SectionClasses);
+
+		// Get all tracks
+		FTopLevelAssetPath TrackClassPath(UMovieSceneTrack::StaticClass());
+		TSet<FTopLevelAssetPath> TrackClasses;
+		AssetRegistryModule.Get().GetDerivedClassNames({ TrackClassPath }, TSet<FTopLevelAssetPath>(), TrackClasses);
+		AllItemClasses.Append(TrackClasses);
+	}
+
+	// Filter to only classes implementing the mixer item interface
+	for (auto It = AllItemClasses.CreateIterator(); It; ++It)
+	{
+		const UClass* Class = FSoftClassPath(It->ToString()).ResolveClass();
+		if (!Class || Class->HasMetaData("Hidden") || !Class->ImplementsInterface(UMovieSceneAnimationMixerItemInterface::StaticClass()))
+		{
+			It.RemoveCurrent();
+		}
+	}
+
+	// Find all registered menu providers
+	TArray<IMovieSceneAnimMixerItemMenuProvider*> Providers = IModularFeatures::Get().GetModularFeatureImplementations<IMovieSceneAnimMixerItemMenuProvider>(IMovieSceneAnimMixerItemMenuProvider::GetModularFeatureName());
 
-		AssetRegistryModule.Get().GetDerivedClassNames({ TargetClassPath }, TSet<FTopLevelAssetPath>(), AllClasses);
+	// Sort providers by priority (lower values appear first)
+	Providers.Sort([](const IMovieSceneAnimMixerItemMenuProvider& A, const IMovieSceneAnimMixerItemMenuProvider& B)
+	{
+		return A.GetMixerItemMenuPriority() < B.GetMixerItemMenuPriority();
+	});
+
+	TSet<const UClass*> HandledClasses;
+
+	TViewModelPtr<FSequencerEditorViewModel> Editor = WeakEditor.Pin();
+	TSharedPtr<ISequencer> Sequencer = Editor ? Editor->GetSequencer() : nullptr;
+
+	FViewModelPtr Model = WeakViewModel.Pin();
+	TViewModelPtr<IObjectBindingExtension> ObjectBinding = Model ? Model->FindAncestorOfType<IObjectBindingExtension>() : nullptr;
+	FGuid ObjectBindingGuid = ObjectBinding ? ObjectBinding->GetObjectGuid() : FGuid();
+
+	TArray<FGuid> ObjectBindings;
+	if (ObjectBindingGuid.IsValid())
+	{
+		ObjectBindings.Add(ObjectBindingGuid);
 	}
 
-	for (auto It = AllClasses.CreateIterator(); It; ++It)
+	// Determine target row index for menu provider callbacks
+	int32 TargetRowIndex = INDEX_NONE;
+	bool bCanAddChildTracks = true;
+
+	if (LayerExtension)
 	{
-		const UClass* Class = FSoftClassPath(It->ToString()).TryLoadClass<UMovieSceneSection>();
+		// Adding to a specific layer - use the layer's row index
+		TargetRowIndex = LayerExtension->GetRowIndex();
 
-		if (!Class || Class->HasMetaData("Hidden") || !Class->ImplementsInterface(IMovieSceneAnimationSectionInterface::UClassType::StaticClass()))
+		// Check if we can add child tracks (only if not on a non-empty layer)
+		UMovieSceneAnimationMixerLayer* Layer = MixerTrack->GetLayer(TargetRowIndex);
+		if (Layer && !Layer->IsEmpty())
 		{
-			It.RemoveCurrent();
+			bCanAddChildTracks = false;
 		}
 	}
+	else
+	{
+		// Adding to the track model - use INDEX_NONE for auto-row assignment
+		TargetRowIndex = INDEX_NONE;
+	}
 
-	if (AllClasses.Num() != 0)
+	// Animation items
+	if (AllItemClasses.Num() > 0)
 	{
 		MenuBuilder.BeginSection(NAME_None, LOCTEXT("AnimationCategoryLabel", "Animation:"));
 
-		for (const FTopLevelAssetPath& ClassPath : AllClasses)
+		// First, call all menu providers
+		for (IMovieSceneAnimMixerItemMenuProvider* Provider : Providers)
 		{
-			const UClass* Class = FSoftClassPath(ClassPath.ToString()).ResolveClass();
-			// Should have been loaded by the validation loop above
-			check(Class);
+			if (Provider)
+			{
+				const UClass* HandledClass = Provider->GetHandledMixerItemClass();
+				if (HandledClass)
+				{
+					// Skip child track providers if we can't add child tracks to this layer
+					const bool bIsTrackProvider = HandledClass->IsChildOf(UMovieSceneTrack::StaticClass());
+					if (bIsTrackProvider && !bCanAddChildTracks)
+					{
+						HandledClasses.Add(HandledClass);
+						continue;
+					}
 
-			MenuBuilder.AddMenuEntry(
-				Class->GetDisplayNameText(),
-				Class->GetToolTipText(),
-				FSlateIcon(),
-				FUIAction(FExecuteAction::CreateLambda(CreateNewSection, ClassPath))
-			);
+					HandledClasses.Add(HandledClass);
+					Provider->PopulateAddMixerItemMenu(MenuBuilder, ObjectBindings, MixerTrack, Sequencer, TargetRowIndex);
+				}
+			}
 		}
 
-		FViewModelPtr Model = WeakViewModel.Pin();
-		TViewModelPtr<IObjectBindingExtension> ObjectBinding = Model ? Model->FindAncestorOfType<IObjectBindingExtension>() : nullptr;
-		if (Track && ObjectBinding)
+		// Then, add generic entries for items without menu providers
+		for (const FTopLevelAssetPath& ClassPath : AllItemClasses)
 		{
-			USkeleton* Skeleton = AcquireSkeletonFromObjectGuid(ObjectBinding->GetObjectGuid(), GetSequencer());
-			if (Skeleton)
+			const UClass* Class = FSoftClassPath(ClassPath.ToString()).ResolveClass();
+			if (!Class || HandledClasses.Contains(Class))
 			{
-				MenuBuilder.AddSubMenu(
-					LOCTEXT("AddAnimationSubMenu", "Animation"),
-					LOCTEXT("AddAnimationSubMenu_Tooltip", "Adds a new animation section for an animation asset at the current time"),
-					FNewMenuDelegate::CreateSP(this, &FAnimationMixerTrackEditor::AddAnimationSubMenu, TArray<FGuid>({ ObjectBinding->GetObjectGuid() }), Skeleton, Track->GetTrack())
-				);
+				continue;
+			}
+
+			// Get display info from interface
+			const UObject* CDO = Class->GetDefaultObject();
+			const IMovieSceneAnimationMixerItemInterface* ItemInterface = Cast<IMovieSceneAnimationMixerItemInterface>(CDO);
+
+			if (ItemInterface)
+			{
+				FText DisplayName = ItemInterface->GetDisplayName();
+				if (DisplayName.IsEmpty())
+				{
+					DisplayName = Class->GetDisplayNameText();
+				}
+
+				FText Description = ItemInterface->GetDescription();
+				if (Description.IsEmpty())
+				{
+					Description = Class->GetToolTipText();
+				}
+
+				// Determine if this is a section or track and use appropriate creation lambda
+				bool bIsSection = Class->IsChildOf(UMovieSceneSection::StaticClass());
+				bool bIsTrack = Class->IsChildOf(UMovieSceneTrack::StaticClass());
+
+				if (bIsSection)
+				{
+					MenuBuilder.AddMenuEntry(
+						DisplayName,
+						Description,
+						ItemInterface->GetIcon(),
+						FUIAction(FExecuteAction::CreateLambda(CreateNewSection, ClassPath))
+					);
+				}
+				else if (bIsTrack)
+				{
+					// Don't show child track options if adding to a non-empty layer
+					if (bCanAddChildTracks)
+					{
+						MenuBuilder.AddMenuEntry(
+							DisplayName,
+							Description,
+							ItemInterface->GetIcon(),
+							FUIAction(FExecuteAction::CreateLambda(CreateNewChildTrack, ObjectBindingGuid, ClassPath))
+						);
+					}
+				}
 			}
 		}
 
diff --git a/Engine/Plugins/Experimental/MovieSceneAnimMixer/Source/MovieSceneAnimMixerEditor/Private/MovieSceneAnimationMixerTrackEditor.h b/Engine/Plugins/Experimental/MovieSceneAnimMixer/Source/MovieSceneAnimMixerEditor/Private/MovieSceneAnimationMixerTrackEditor.h
--- a/Engine/Plugins/Experimental/MovieSceneAnimMixer/Source/MovieSceneAnimMixerEditor/Private/MovieSceneAnimationMixerTrackEditor.h
+++ b/Engine/Plugins/Experimental/MovieSceneAnimMixer/Source/MovieSceneAnimMixerEditor/Private/MovieSceneAnimationMixerTrackEditor.h
@@ -6,6 +6,7 @@
 #include "Templates/SubclassOf.h"
 #include "MovieSceneTrackEditor.h"
 #include "TrackEditors/CommonAnimationTrackEditor.h"
+#include "IMovieSceneAnimMixerItemMenuProvider.h"
 
 DECLARE_DELEGATE_RetVal_ThreeParams(TSharedRef<ISequencerSection>, FOnMakeSectionInterfaceDelegate, UMovieSceneSection&, UMovieSceneTrack&, FGuid);
 
@@ -14,6 +15,7 @@ namespace UE::Sequencer
 
 class FAnimationMixerTrackEditor
 	: public FCommonAnimationTrackEditor
+	, public IMovieSceneAnimMixerItemMenuProvider
 {
 public:
 
@@ -40,11 +42,19 @@ public:
 	virtual TSharedRef<ISequencerSection> MakeSectionInterface(UMovieSceneSection& SectionObject, UMovieSceneTrack& Track, FGuid ObjectBinding) override;
 	virtual TSharedPtr<SWidget> BuildOutlinerColumnWidget(const FBuildColumnWidgetParams& Params, const FName& ColumnName) override;
 	virtual void BuildObjectBindingTrackMenu(FMenuBuilder& MenuBuilder, const TArray<FGuid>& ObjectBindings, const UClass* ObjectClass) override;
+	virtual void OnInitialize() override;
+	virtual void OnRelease() override;
 	virtual bool SupportsSequence(UMovieSceneSequence* InSequence) const override;
 	virtual bool SupportsType(TSubclassOf<UMovieSceneTrack> Type) const override;
 	virtual const FSlateBrush* GetIconBrush() const override;
 	virtual TSubclassOf<UMovieSceneCommonAnimationTrack> GetTrackClass() const override;
 
+	// IMovieSceneAnimMixerItemMenuProvider interface
+	virtual const UClass* GetHandledMixerItemClass() const override;
+	virtual void PopulateAddMixerItemMenu(FMenuBuilder& MenuBuilder, TArray<FGuid> ObjectBindings, UMovieSceneTrack* Track, TSharedPtr<ISequencer> Sequencer, int32 RowIndex) override;
+	virtual void PopulateObjectBindingAnimationMenu(FMenuBuilder& MenuBuilder, const TArray<FGuid>& ObjectBindings, const UClass* ObjectClass, bool bIsInsideSubmenu) override;
+	virtual int32 GetObjectBindingAnimationMenuPriority() const override { return 0; } // Highest priority - appears first
+
 private:
 	using FCommonAnimationTrackEditor::BuildAddAnimationSubMenu;
 
diff --git a/Engine/Plugins/Experimental/MovieSceneMixedControlRig/Source/MovieSceneMixedControlRig/Private/Systems/MovieSceneMixedControlRigSystem.cpp b/Engine/Plugins/Experimental/MovieSceneMixedControlRig/Source/MovieSceneMixedControlRig/Private/Systems/MovieSceneMixedControlRigSystem.cpp
--- a/Engine/Plugins/Experimental/MovieSceneMixedControlRig/Source/MovieSceneMixedControlRig/Private/Systems/MovieSceneMixedControlRigSystem.cpp
+++ b/Engine/Plugins/Experimental/MovieSceneMixedControlRig/Source/MovieSceneMixedControlRig/Private/Systems/MovieSceneMixedControlRigSystem.cpp
@@ -80,6 +80,7 @@ struct FControlRigMixerMutation : IMovieSceneEntityMutation
 		TComponentReader<FRootInstanceHandle> RootInstanceHandles = Allocation->ReadComponents(BuiltInComponents->RootInstanceHandle);
 		TComponentReader<FGuid> ObjectBindingIDs = Allocation->ReadComponents(BuiltInComponents->GenericObjectBinding);
 		TOptionalComponentReader<FMovieSceneSequenceID> OptSequenceIDs = Allocation->TryReadComponents(BuiltInComponents->SequenceID);
+		TOptionalComponentReader<TObjectPtr<UMovieSceneAnimationMixerLayer>> OptMixerLayers = Allocation->TryReadComponents(AnimMixerComponents->MixerLayer);
 		
 		TComponentWriter<TInstancedStruct<FMovieSceneMixedAnimationTarget>> OutAnimTargets = Allocation->WriteComponents(AnimMixerComponents->Target, WriteContext);
 		TComponentWriter<TSharedPtr<FAnimNextEvaluationTask>> OutTasks = Allocation->WriteComponents(AnimMixerComponents->Task, WriteContext);
@@ -165,11 +166,9 @@ struct FControlRigMixerMutation : IMovieSceneEntityMutation
 								}
 							}
 						}
-						UMovieSceneAnimationMixerTrack* MixerTrack = Cast<UMovieSceneAnimationMixerTrack>(ControlRigSource.Track->GetOuter());
 
-						// If we have a mixer track, the animation decoration may already have added our target and priority
-
-						if (MixerTrack == nullptr || !MixerTrack->MixedAnimationTarget.IsValid())
+						// If we haven't had a mixer layer set, then set default anim target/priority
+						if (!OptMixerLayers || !OptMixerLayers[Index])
 						{
 							// Set a default target- Automatic will resolve to the appropriate place, or the Anim Mixer track will override it.
 							OutAnimTargets[Index] = TInstancedStruct<FMovieSceneMixedAnimationTarget>::Make();
diff --git a/Engine/Plugins/Experimental/MovieScenePoseSearchTracks/Source/MovieScenePoseSearchTracks/MovieScenePoseSearchTracks.Build.cs b/Engine/Plugins/Experimental/MovieScenePoseSearchTracks/Source/MovieScenePoseSearchTracks/MovieScenePoseSearchTracks.Build.cs
--- a/Engine/Plugins/Experimental/MovieScenePoseSearchTracks/Source/MovieScenePoseSearchTracks/MovieScenePoseSearchTracks.Build.cs
+++ b/Engine/Plugins/Experimental/MovieScenePoseSearchTracks/Source/MovieScenePoseSearchTracks/MovieScenePoseSearchTracks.Build.cs
@@ -19,7 +19,8 @@ namespace UnrealBuildTool.Rules
 					"UAFPoseSearch",
 					"UAFAnimGraph",
 					"AnimGraphRuntime",
-					"PoseSearch"
+					"PoseSearch",
+					"SlateCore"
 				}
 			);
 		}
diff --git a/Engine/Plugins/Experimental/MovieScenePoseSearchTracks/Source/MovieScenePoseSearchTracks/Public/Sections/MovieSceneStitchAnimSection.h b/Engine/Plugins/Experimental/MovieScenePoseSearchTracks/Source/MovieScenePoseSearchTracks/Public/Sections/MovieSceneStitchAnimSection.h
--- a/Engine/Plugins/Experimental/MovieScenePoseSearchTracks/Source/MovieScenePoseSearchTracks/Public/Sections/MovieSceneStitchAnimSection.h
+++ b/Engine/Plugins/Experimental/MovieScenePoseSearchTracks/Source/MovieScenePoseSearchTracks/Public/Sections/MovieSceneStitchAnimSection.h
@@ -109,7 +109,7 @@ UCLASS( MinimalAPI, DisplayName="Animation Stitch")
 class UMovieSceneStitchAnimSection
 	: public UMovieSceneSection
 	, public IMovieSceneEntityProvider
-	, public IMovieSceneAnimationSectionInterface
+	, public IMovieSceneAnimationMixerItemInterface
 {
 	GENERATED_BODY()
 
@@ -150,12 +150,7 @@ public:
 	virtual EMovieSceneChannelProxyType CacheChannelProxy() override;
 	virtual UObject* GetSourceObject() const override;
 
-	virtual int32 GetRowSortOrder() const override
-	{
-		return 0;
-	}
-
-	virtual FColor GetMixerSectionTint() const override
+	virtual FColor GetMixerItemTint() const override
 	{
 		return MixerTintOverride;
 	}
diff --git a/Engine/Plugins/Experimental/MovieScenePoseSearchTracks/Source/MovieScenePoseSearchTracksEditor/Private/TrackEditors/StitchAnimTrackEditor.cpp b/Engine/Plugins/Experimental/MovieScenePoseSearchTracks/Source/MovieScenePoseSearchTracksEditor/Private/TrackEditors/StitchAnimTrackEditor.cpp
--- a/Engine/Plugins/Experimental/MovieScenePoseSearchTracks/Source/MovieScenePoseSearchTracksEditor/Private/TrackEditors/StitchAnimTrackEditor.cpp
+++ b/Engine/Plugins/Experimental/MovieScenePoseSearchTracks/Source/MovieScenePoseSearchTracksEditor/Private/TrackEditors/StitchAnimTrackEditor.cpp
@@ -14,12 +14,14 @@
 #include "PoseSearch/PoseSearchDatabase.h"
 #include "Framework/MultiBox/MultiBoxBuilder.h"
 #include "MVVM/Extensions/ITrackExtension.h"
+#include "MVVM/ViewModels/ViewModelIterators.h"
 #include "IContentBrowserSingleton.h"
 #include "ContentBrowserModule.h"
 #include "SequencerSettings.h"
 #include "MVVM/Views/ViewUtilities.h"
 #include "DragAndDrop/AssetDragDropOp.h"
 #include "MovieSceneToolHelpers.h"
+#include "IMovieSceneAnimMixerItemMenuProvider.h"
 
 namespace StitchAnimEditorConstants
 {
@@ -270,7 +272,8 @@ bool FStitchAnimTrackEditor::HandleAssetAdded(UObject* Asset, const FGuid& Targe
 				const FScopedTransaction Transaction(LOCTEXT("AddStitchedAnim_Transaction", "Add Stitched Animation"));
 
 				int32 RowIndex = INDEX_NONE;
-				AnimatablePropertyChanged(FOnKeyProperty::CreateRaw(this, &FStitchAnimTrackEditor::AddKeyInternal, Object, PoseSearchDatabase, Track, RowIndex));
+				TFunction<void(UMovieSceneSection*)> SectionCallback = nullptr;
+				AnimatablePropertyChanged(FOnKeyProperty::CreateRaw(this, &FStitchAnimTrackEditor::AddKeyInternal, Object, PoseSearchDatabase, Track, RowIndex, SectionCallback));
 
 				return true;
 			}
@@ -279,7 +282,26 @@ bool FStitchAnimTrackEditor::HandleAssetAdded(UObject* Asset, const FGuid& Targe
 	return false;
 }
 
+
+void FStitchAnimTrackEditor::OnInitialize()
+{
+	// Register as animation mixer item menu provider
+	IModularFeatures::Get().RegisterModularFeature(IMovieSceneAnimMixerItemMenuProvider::GetModularFeatureName(), this);
+}
+
+void FStitchAnimTrackEditor::OnRelease()
+{
+	// Unregister as animation mixer item menu provider
+	IModularFeatures::Get().UnregisterModularFeature(IMovieSceneAnimMixerItemMenuProvider::GetModularFeatureName(), this);
+}
+
 void FStitchAnimTrackEditor::BuildObjectBindingTrackMenu(FMenuBuilder& MenuBuilder, const TArray<FGuid>& ObjectBindings, const UClass* ObjectClass)
+{
+	// Stitch Animation menu entries are now added via PopulateObjectBindingAnimationMenu,
+	// which is called by FCommonAnimationTrackEditor::BuildObjectBindingTrackMenu
+}
+
+void FStitchAnimTrackEditor::PopulateObjectBindingAnimationMenu(FMenuBuilder& MenuBuilder, const TArray<FGuid>& ObjectBindings, const UClass* ObjectClass, bool bIsInsideSubmenu)
 {
 	using namespace UE::PoseSearch;
 	if (ObjectClass->IsChildOf(USkeletalMeshComponent::StaticClass()) || ObjectClass->IsChildOf(AActor::StaticClass()) || ObjectClass->IsChildOf(UChildActorComponent::StaticClass()))
@@ -300,18 +322,35 @@ void FStitchAnimTrackEditor::BuildObjectBindingTrackMenu(FMenuBuilder& MenuBuild
 			if (AssetDataList.Num())
 			{
 				UMovieSceneTrack* Track = nullptr;
+				TFunction<void(UMovieSceneSection*)> SectionCallback = nullptr;
 
 				MenuBuilder.AddSubMenu(
 					LOCTEXT("AddStitchAnimation", "Stitch Animation"), NSLOCTEXT("Sequencer", "AddStitchAnimationTooltip", "Adds a stitch animation track."),
-					FNewMenuDelegate::CreateRaw(this, &FStitchAnimTrackEditor::AddAnimationSubMenu, ObjectBindings, Skeleton, Track)
+					FNewMenuDelegate::CreateRaw(this, &FStitchAnimTrackEditor::AddAnimationSubMenu, ObjectBindings, Skeleton, Track, SectionCallback)
 				);
 			}
 		}
 	}
 }
 
-TSharedRef<SWidget> FStitchAnimTrackEditor::BuildAddAnimationSubMenu(FGuid ObjectBinding, USkeleton* Skeleton, UE::Sequencer::TWeakViewModelPtr<UE::Sequencer::ITrackExtension> WeakTrackModel)
+TSharedRef<SWidget> FStitchAnimTrackEditor::BuildAddAnimationSubMenu(FGuid ObjectBinding, USkeleton* Skeleton, UE::Sequencer::TWeakViewModelPtr<UE::Sequencer::FViewModel> WeakViewModel, TFunction<void(UMovieSceneSection*)> InSectionCallback)
 {
+	using namespace UE::Sequencer;
+	TViewModelPtr<FViewModel> ViewModel = WeakViewModel.Pin();
+	if (!ViewModel)
+	{
+		return SNullWidget::NullWidget;
+	}
+
+	TViewModelPtr<ITrackExtension> TrackModel = ViewModel->FindAncestorOfType<ITrackExtension>(true);
+
+	if (!TrackModel)
+	{
+		return SNullWidget::NullWidget;
+	}
+
+	UMovieSceneTrack* Track = TrackModel->GetTrack();
+
 	FMenuBuilder MenuBuilder(true, nullptr);
 
 	TArray<FGuid> ObjectBindings;
@@ -319,26 +358,26 @@ TSharedRef<SWidget> FStitchAnimTrackEditor::BuildAddAnimationSubMenu(FGuid Objec
 
 	MenuBuilder.BeginSection(NAME_None, LOCTEXT("AddStitchAnimation_Label", "Add Stitch Animation"));
 	{
-		AddAnimationSubMenu(MenuBuilder, ObjectBindings, Skeleton, WeakTrackModel.Pin()->GetTrack());
+		AddAnimationSubMenu(MenuBuilder, ObjectBindings, Skeleton, Track, InSectionCallback);
 	}
 	MenuBuilder.EndSection();
 
 	return MenuBuilder.MakeWidget();
 }
 
-TSharedRef<SWidget> FStitchAnimTrackEditor::BuildAnimationSubMenu(FGuid ObjectBinding, USkeleton* Skeleton, UMovieSceneTrack* Track)
+TSharedRef<SWidget> FStitchAnimTrackEditor::BuildAnimationSubMenu(FGuid ObjectBinding, USkeleton* Skeleton, UMovieSceneTrack* Track, TFunction<void(UMovieSceneSection*)> InSectionCallback)
 {
 	FMenuBuilder MenuBuilder(true, nullptr);
 
 	TArray<FGuid> ObjectBindings;
 	ObjectBindings.Add(ObjectBinding);
 
-	AddAnimationSubMenu(MenuBuilder, ObjectBindings, Skeleton, Track);
+	AddAnimationSubMenu(MenuBuilder, ObjectBindings, Skeleton, Track, InSectionCallback);
 
 	return MenuBuilder.MakeWidget();
 }
 
-void FStitchAnimTrackEditor::AddAnimationSubMenu(FMenuBuilder& MenuBuilder, TArray<FGuid> ObjectBindings, USkeleton* Skeleton, UMovieSceneTrack* Track)
+void FStitchAnimTrackEditor::AddAnimationSubMenu(FMenuBuilder& MenuBuilder, TArray<FGuid> ObjectBindings, USkeleton* Skeleton, UMovieSceneTrack* Track, TFunction<void(UMovieSceneSection*)> InSectionCallback)
 {
 	using namespace UE::PoseSearch;
 	TSharedPtr<ISequencer> SequencerPtr = GetSequencer();
@@ -346,8 +385,8 @@ void FStitchAnimTrackEditor::AddAnimationSubMenu(FMenuBuilder& MenuBuilder, TArr
 
 	FAssetPickerConfig AssetPickerConfig;
 	{
-		AssetPickerConfig.OnAssetSelected = FOnAssetSelected::CreateRaw(this, &FStitchAnimTrackEditor::OnAnimationDatabaseAssetSelected, ObjectBindings, Track);
-		AssetPickerConfig.OnAssetEnterPressed = FOnAssetEnterPressed::CreateRaw(this, &FStitchAnimTrackEditor::OnAnimationDatabaseAssetEnterPressed, ObjectBindings, Track);
+		AssetPickerConfig.OnAssetSelected = FOnAssetSelected::CreateRaw(this, &FStitchAnimTrackEditor::OnAnimationDatabaseAssetSelected, ObjectBindings, Track, InSectionCallback);
+		AssetPickerConfig.OnAssetEnterPressed = FOnAssetEnterPressed::CreateRaw(this, &FStitchAnimTrackEditor::OnAnimationDatabaseAssetEnterPressed, ObjectBindings, Track, InSectionCallback);
 		AssetPickerConfig.bAllowNullSelection = false;
 		AssetPickerConfig.bAddFilterUI = true;
 		AssetPickerConfig.bShowTypeInColumnView = false;
@@ -373,7 +412,7 @@ void FStitchAnimTrackEditor::AddAnimationSubMenu(FMenuBuilder& MenuBuilder, TArr
 	MenuBuilder.AddWidget(MenuEntry.ToSharedRef(), FText::GetEmpty(), true);
 }
 
-void FStitchAnimTrackEditor::OnAnimationDatabaseAssetSelected(const FAssetData& AssetData, TArray<FGuid> ObjectBindings, UMovieSceneTrack* Track)
+void FStitchAnimTrackEditor::OnAnimationDatabaseAssetSelected(const FAssetData& AssetData, TArray<FGuid> ObjectBindings, UMovieSceneTrack* Track, TFunction<void(UMovieSceneSection*)> InSectionCallback)
 {
 	using namespace UE::PoseSearch;
 	FSlateApplication::Get().DismissAllMenus();
@@ -391,21 +430,21 @@ void FStitchAnimTrackEditor::OnAnimationDatabaseAssetSelected(const FAssetData&
 		{
 			UObject* Object = SequencerPtr->FindSpawnedObjectOrTemplate(ObjectBinding);
 			int32 RowIndex = INDEX_NONE;
-			AnimatablePropertyChanged(FOnKeyProperty::CreateRaw(this, &FStitchAnimTrackEditor::AddKeyInternal, Object, PoseSearchDatabase, Track, RowIndex));
+			AnimatablePropertyChanged(FOnKeyProperty::CreateRaw(this, &FStitchAnimTrackEditor::AddKeyInternal, Object, PoseSearchDatabase, Track, RowIndex, InSectionCallback));
 		}
 	}
 }
 
-void FStitchAnimTrackEditor::OnAnimationDatabaseAssetEnterPressed(const TArray<FAssetData>& AssetData, TArray<FGuid> ObjectBindings, UMovieSceneTrack* Track)
+void FStitchAnimTrackEditor::OnAnimationDatabaseAssetEnterPressed(const TArray<FAssetData>& AssetData, TArray<FGuid> ObjectBindings, UMovieSceneTrack* Track, TFunction<void(UMovieSceneSection*)> InSectionCallback)
 {
 	if (AssetData.Num() > 0)
 	{
-		OnAnimationDatabaseAssetSelected(AssetData[0].GetAsset(), ObjectBindings, Track);
+		OnAnimationDatabaseAssetSelected(AssetData[0].GetAsset(), ObjectBindings, Track, InSectionCallback);
 	}
 }
 
 
-FKeyPropertyResult FStitchAnimTrackEditor::AddKeyInternal(FFrameNumber KeyTime, UObject* Object, UPoseSearchDatabase* PoseSearchDatabase, UMovieSceneTrack* Track, int32 RowIndex)
+FKeyPropertyResult FStitchAnimTrackEditor::AddKeyInternal(FFrameNumber KeyTime, UObject* Object, UPoseSearchDatabase* PoseSearchDatabase, UMovieSceneTrack* Track, int32 RowIndex, TFunction<void(UMovieSceneSection*)> InSectionCallback)
 {
 	FKeyPropertyResult KeyPropertyResult;
 
@@ -415,21 +454,20 @@ FKeyPropertyResult FStitchAnimTrackEditor::AddKeyInternal(FFrameNumber KeyTime,
 	if (ObjectHandle.IsValid())
 	{
 		UMovieScene* MovieScene = GetSequencer()->GetFocusedMovieSceneSequence()->GetMovieScene();
-		UMovieSceneStitchAnimTrack* StitchAnimTrack = Cast<UMovieSceneStitchAnimTrack>(Track);
-		FMovieSceneBinding* Binding = MovieScene->FindBinding(ObjectHandle);
 
-		// Add a track if no track was specified or if the track specified doesn't belong to the tracks of the targeted guid
-		if (!StitchAnimTrack || (Binding && !Binding->GetTracks().Contains(StitchAnimTrack)))
+		// If a callback is provided (mixer case), create the section directly and let the callback handle adding it
+		if (InSectionCallback)
 		{
-			StitchAnimTrack = CastChecked<UMovieSceneStitchAnimTrack>(AddTrack(MovieScene, ObjectHandle, UMovieSceneStitchAnimTrack::StaticClass(), NAME_None), ECastCheckedType::NullAllowed);
-			KeyPropertyResult.bTrackCreated = true;
-		}
+			// Create section (outer will be set properly when callback adds it to track)
+			UMovieSceneStitchAnimSection* NewSection = NewObject<UMovieSceneStitchAnimSection>(GetTransientPackage(), NAME_None, RF_Transactional);
+			NewSection->StitchDatabase = PoseSearchDatabase;
 
-		if (ensure(StitchAnimTrack))
-		{
-			StitchAnimTrack->Modify();
+			// Set up default section timing (duplicating what AddNewAnimationOnRow does)
+			NewSection->SetRange(TRange<FFrameNumber>(KeyTime, KeyTime + (int32)FMath::FloorToInt(3.0f * MovieScene->GetTickResolution().AsDecimal())));
+
+			// Call the callback to add the section to the track
+			InSectionCallback(NewSection);
 
-			UMovieSceneStitchAnimSection* NewSection = Cast<UMovieSceneStitchAnimSection>(StitchAnimTrack->AddNewAnimationOnRow(KeyTime, PoseSearchDatabase, RowIndex));
 			KeyPropertyResult.bTrackModified = true;
 			KeyPropertyResult.SectionsCreated.Add(NewSection);
 
@@ -437,18 +475,95 @@ FKeyPropertyResult FStitchAnimTrackEditor::AddKeyInternal(FFrameNumber KeyTime,
 			GetSequencer()->SelectSection(NewSection);
 			GetSequencer()->ThrobSectionSelection();
 		}
+		else
+		{
+			// Standard path: create dedicated track and add section normally
+			UMovieSceneStitchAnimTrack* StitchAnimTrack = Cast<UMovieSceneStitchAnimTrack>(Track);
+			FMovieSceneBinding* Binding = MovieScene->FindBinding(ObjectHandle);
+
+			// Add a track if no track was specified or if the track specified doesn't belong to the tracks of the targeted guid
+			if (!StitchAnimTrack || (Binding && !Binding->GetTracks().Contains(StitchAnimTrack)))
+			{
+				StitchAnimTrack = CastChecked<UMovieSceneStitchAnimTrack>(AddTrack(MovieScene, ObjectHandle, UMovieSceneStitchAnimTrack::StaticClass(), NAME_None), ECastCheckedType::NullAllowed);
+				KeyPropertyResult.bTrackCreated = true;
+			}
+
+			if (ensure(StitchAnimTrack))
+			{
+				StitchAnimTrack->Modify();
+
+				UMovieSceneStitchAnimSection* NewSection = Cast<UMovieSceneStitchAnimSection>(StitchAnimTrack->AddNewAnimationOnRow(KeyTime, PoseSearchDatabase, RowIndex));
+				KeyPropertyResult.bTrackModified = true;
+				KeyPropertyResult.SectionsCreated.Add(NewSection);
+
+				GetSequencer()->EmptySelection();
+				GetSequencer()->SelectSection(NewSection);
+				GetSequencer()->ThrobSectionSelection();
+			}
+		}
 	}
 
 	return KeyPropertyResult;
 }
 
+const UClass* FStitchAnimTrackEditor::GetHandledMixerItemClass() const
+{
+	return UMovieSceneStitchAnimSection::StaticClass();
+}
+
+void FStitchAnimTrackEditor::PopulateAddMixerItemMenu(FMenuBuilder& MenuBuilder, TArray<FGuid> ObjectBindings, UMovieSceneTrack* Track, TSharedPtr<ISequencer> Sequencer, int32 RowIndex)
+{
+	using namespace UE::PoseSearch;
+
+	if (ObjectBindings.Num() == 0 || !Sequencer.IsValid() || !Track)
+	{
+		return;
+	}
+
+	USkeleton* Skeleton = AcquireSkeletonFromObjectGuid(ObjectBindings[0], Sequencer);
+	if (!Skeleton)
+	{
+		return;
+	}
+
+	// Create a callback that will add the section to the mixer track at the specified row or with auto row
+	auto SectionCallback = [Track, RowIndex](UMovieSceneSection* Section)
+	{
+		if (UMovieSceneCommonAnimationTrack* CommonTrack = Cast<UMovieSceneCommonAnimationTrack>(Track))
+		{
+			if (RowIndex != INDEX_NONE)
+			{
+				// Use the specified row index
+				Section->SetRowIndex(RowIndex);
+				CommonTrack->AddSection(*Section);
+			}
+			else
+			{
+				// Use auto-row assignment
+				CommonTrack->AddSectionWithAutoRow(Section);
+			}
+		}
+	};
+
+	// Wrap the stitch animation menu in a submenu
+	MenuBuilder.AddSubMenu(
+		LOCTEXT("AddStitchAnimation", "Stitch Animation"),
+		LOCTEXT("AddStitchAnimationTooltip", "Adds a stitch animation section"),
+		FNewMenuDelegate::CreateLambda([this, ObjectBindings, Skeleton, Track, SectionCallback](FMenuBuilder& SubMenuBuilder)
+		{
+			AddAnimationSubMenu(SubMenuBuilder, ObjectBindings, Skeleton, Track, SectionCallback);
+		})
+	);
+}
+
 TSharedPtr<SWidget> FStitchAnimTrackEditor::BuildOutlinerEditWidget(const FGuid& ObjectBinding, UMovieSceneTrack* Track, const FBuildEditWidgetParams& Params)
 {
 	USkeleton* Skeleton = AcquireSkeletonFromObjectGuid(ObjectBinding, GetSequencer());
 
 	if (Skeleton)
 	{
-		FOnGetContent HandleGetAddButtonContent = FOnGetContent::CreateSP(this, &FStitchAnimTrackEditor::BuildAddAnimationSubMenu, ObjectBinding, Skeleton, Params.TrackModel.AsWeak());
+		TFunction<void(UMovieSceneSection*)> SectionCallback = nullptr;
+		FOnGetContent HandleGetAddButtonContent = FOnGetContent::CreateSP(this, &FStitchAnimTrackEditor::BuildAddAnimationSubMenu, ObjectBinding, Skeleton, Params.ViewModel.AsWeak(), SectionCallback);
 		return UE::Sequencer::MakeAddButton(LOCTEXT("AnimationText", "Stitch Animation"), HandleGetAddButtonContent, Params.ViewModel);
 	}
 	else
@@ -567,7 +682,8 @@ FReply FStitchAnimTrackEditor::OnDrop(const FDragDropEvent& DragDropEvent, const
 			{
 				UObject* BoundObject = SequencerPtr.IsValid() ? SequencerPtr->FindSpawnedObjectOrTemplate(DragDropParams.TargetObjectGuid) : nullptr;
 
-				AnimatablePropertyChanged(FOnKeyProperty::CreateRaw(this, &FStitchAnimTrackEditor::AddKeyInternal, BoundObject, PoseSearchDatabase, DragDropParams.Track.Get(), DragDropParams.RowIndex));
+				TFunction<void(UMovieSceneSection*)> SectionCallback = nullptr;
+				AnimatablePropertyChanged(FOnKeyProperty::CreateRaw(this, &FStitchAnimTrackEditor::AddKeyInternal, BoundObject, PoseSearchDatabase, DragDropParams.Track.Get(), DragDropParams.RowIndex, SectionCallback));
 
 				bAnyDropped = true;
 			}
