Sequencer: Anim Mixer target selection.

Add target selection as part of the anim mixer track creation workflow, and also add a combo box for target selection on the anim mixer track itself.

Required making public the combo box for uaf variable selection as we reuse it in the combo box.

#Synced-CL 49155801

diff --git a/Engine/Plugins/Experimental/MovieSceneAnimMixer/Source/MovieSceneAnimMixer/Public/Systems/MovieSceneAnimBlueprintTargetSystem.h b/Engine/Plugins/Experimental/MovieSceneAnimMixer/Source/MovieSceneAnimMixer/Public/Systems/MovieSceneAnimBlueprintTargetSystem.h
--- a/Engine/Plugins/Experimental/MovieSceneAnimMixer/Source/MovieSceneAnimMixer/Public/Systems/MovieSceneAnimBlueprintTargetSystem.h
+++ b/Engine/Plugins/Experimental/MovieSceneAnimMixer/Source/MovieSceneAnimMixer/Public/Systems/MovieSceneAnimBlueprintTargetSystem.h
@@ -34,8 +34,15 @@ struct MOVIESCENEANIMMIXER_API FMovieSceneAnimBlueprintTarget : public FMovieSce
 	FMovieSceneAnimBlueprintTarget()
 		: BlueprintNodeName(FAnimNode_SequencerMixerTarget::DefaultTargetName)
 	{
-		
+
+	}
+
+#if WITH_EDITOR
+	virtual FText GetShortDisplayName() const override
+	{
+		return FText::FromName(BlueprintNodeName);
 	}
+#endif
 };
 
 UCLASS(MinimalAPI)
diff --git a/Engine/Plugins/Experimental/MovieSceneAnimMixer/Source/MovieSceneAnimMixer/Public/Systems/MovieSceneAnimInstanceTargetSystem.h b/Engine/Plugins/Experimental/MovieSceneAnimMixer/Source/MovieSceneAnimMixer/Public/Systems/MovieSceneAnimInstanceTargetSystem.h
--- a/Engine/Plugins/Experimental/MovieSceneAnimMixer/Source/MovieSceneAnimMixer/Public/Systems/MovieSceneAnimInstanceTargetSystem.h
+++ b/Engine/Plugins/Experimental/MovieSceneAnimMixer/Source/MovieSceneAnimMixer/Public/Systems/MovieSceneAnimInstanceTargetSystem.h
@@ -26,6 +26,13 @@ struct MOVIESCENEANIMMIXER_API FMovieSceneAnimInstanceTarget : public FMovieScen
 	{
 		return GetTypeHash(FMovieSceneAnimInstanceTarget::StaticStruct());
 	}
+
+#if WITH_EDITOR
+	virtual FText GetShortDisplayName() const override
+	{
+		return NSLOCTEXT("MovieSceneAnimInstanceTarget", "CustomAnimInstance", "Custom Anim Instance");
+	}
+#endif
 };
 
 
diff --git a/Engine/Plugins/Experimental/MovieSceneAnimMixer/Source/MovieSceneAnimMixer/Public/Systems/MovieSceneAnimNextTargetSystem.h b/Engine/Plugins/Experimental/MovieSceneAnimMixer/Source/MovieSceneAnimMixer/Public/Systems/MovieSceneAnimNextTargetSystem.h
--- a/Engine/Plugins/Experimental/MovieSceneAnimMixer/Source/MovieSceneAnimMixer/Public/Systems/MovieSceneAnimNextTargetSystem.h
+++ b/Engine/Plugins/Experimental/MovieSceneAnimMixer/Source/MovieSceneAnimMixer/Public/Systems/MovieSceneAnimNextTargetSystem.h
@@ -35,7 +35,20 @@ struct FMovieSceneAnimMixerEvaluationModifier : public UE::UAF::IEvaluationModif
 USTRUCT(meta=(DisplayName="UAF Module Injection"))
 struct MOVIESCENEANIMMIXER_API FMovieSceneAnimNextInjectionTarget : public FMovieSceneMixedAnimationTarget
 {
-	GENERATED_BODY() 
+	GENERATED_BODY()
+
+	FMovieSceneAnimNextInjectionTarget(const FAnimNextSoftVariableReference& InVariableReference)
+		: InjectionSite(InVariableReference)
+	{
+
+	}
+
+	FMovieSceneAnimNextInjectionTarget()
+		: InjectionSite()
+	{
+
+	}
+
 
 	UPROPERTY()
 	FName InjectionSiteName_DEPRECATED;
@@ -53,6 +66,17 @@ struct MOVIESCENEANIMMIXER_API FMovieSceneAnimNextInjectionTarget : public FMovi
 	bool Serialize(FArchive& Ar);
 	void PostSerialize(const FArchive& Ar);
 #endif
+
+#if WITH_EDITOR
+	virtual FText GetShortDisplayName() const override
+	{
+		if (!InjectionSite.IsNone())
+		{
+			return FText::FromName(InjectionSite.GetName());
+		}
+		return NSLOCTEXT("MovieSceneAnimNextInjectionTarget", "UAFModuleInjection", "UAF Module Injection");
+	}
+#endif
 };
 
 struct FMovieSceneAnimNextTargetData
diff --git a/Engine/Plugins/Experimental/MovieSceneAnimMixer/Source/MovieSceneAnimMixerEditor/MovieSceneAnimMixerEditor.Build.cs b/Engine/Plugins/Experimental/MovieSceneAnimMixer/Source/MovieSceneAnimMixerEditor/MovieSceneAnimMixerEditor.Build.cs
--- a/Engine/Plugins/Experimental/MovieSceneAnimMixer/Source/MovieSceneAnimMixerEditor/MovieSceneAnimMixerEditor.Build.cs
+++ b/Engine/Plugins/Experimental/MovieSceneAnimMixer/Source/MovieSceneAnimMixerEditor/MovieSceneAnimMixerEditor.Build.cs
@@ -31,6 +31,9 @@ namespace UnrealBuildTool.Rules
 					"SequencerCore",
 					"Slate",
 					"SlateCore",
+					"UAF",
+					"UAFEditor",
+					"UAFUncookedOnly",
 					"UnrealEd",
 				}
 			);
diff --git a/Engine/Plugins/Experimental/MovieSceneAnimMixer/Source/MovieSceneAnimMixerEditor/Private/AnimMixerTargetMenuProviders.cpp b/Engine/Plugins/Experimental/MovieSceneAnimMixer/Source/MovieSceneAnimMixerEditor/Private/AnimMixerTargetMenuProviders.cpp
new file mode 100644
--- /dev/null
+++ b/Engine/Plugins/Experimental/MovieSceneAnimMixer/Source/MovieSceneAnimMixerEditor/Private/AnimMixerTargetMenuProviders.cpp
@@ -0,0 +1,238 @@
+// Copyright Epic Games, Inc. All Rights Reserved.
+
+#include "AnimMixerTargetMenuProviders.h"
+
+#include "AnimGraphNode_SequencerMixerTarget.h"
+#include "Animation/AnimBlueprint.h"
+#include "Animation/AnimBlueprintGeneratedClass.h"
+#include "Animation/AnimInstance.h"
+#include "Component/AnimNextComponent.h"
+#include "Components/SkeletalMeshComponent.h"
+#include "Framework/MultiBox/MultiBoxBuilder.h"
+#include "GameFramework/Actor.h"
+#include "MovieSceneMixedAnimationTarget.h"
+#include "Param/ParamCompatibility.h"
+#include "Param/ParamType.h"
+#include "Param/ParamUtils.h"
+#include "Variables/SVariablePickerCombo.h"
+#include "Systems/MovieSceneAnimBlueprintTargetSystem.h"
+#include "Systems/MovieSceneAnimInstanceTargetSystem.h"
+#include "Systems/MovieSceneAnimNextTargetSystem.h"
+#include "Variables/AnimNextSoftVariableReference.h"
+#include "Variables/AnimNextVariableReference.h"
+#include "Variables/VariablePickerArgs.h"
+
+#define LOCTEXT_NAMESPACE "AnimMixerTargetMenuProviders"
+
+//------------------------------------------------------------------------------
+// FAutomaticTargetMenuProvider
+//------------------------------------------------------------------------------
+
+UScriptStruct* FAutomaticTargetMenuProvider::GetHandledTargetStructType() const
+{
+	return FMovieSceneMixedAnimationTarget::StaticStruct();
+}
+
+void FAutomaticTargetMenuProvider::PopulateTargetMenu(
+	FMenuBuilder& MenuBuilder,
+	TObjectPtr<UObject> BoundObject,
+	TFunction<void(TInstancedStruct<FMovieSceneMixedAnimationTarget>)> OnTargetSelected)
+{
+	MenuBuilder.AddMenuEntry(
+		LOCTEXT("AutomaticTarget", "Automatic"),
+		LOCTEXT("AutomaticTargetTooltip", "Automatically determine the animation target based on the object's configuration."),
+		FSlateIcon(),
+		FUIAction(
+			FExecuteAction::CreateLambda([OnTargetSelected]()
+			{
+				OnTargetSelected(TInstancedStruct<FMovieSceneMixedAnimationTarget>::Make());
+			})
+		)
+	);
+}
+
+//------------------------------------------------------------------------------
+// FAnimInstanceTargetMenuProvider
+//------------------------------------------------------------------------------
+
+UScriptStruct* FAnimInstanceTargetMenuProvider::GetHandledTargetStructType() const
+{
+	return FMovieSceneAnimInstanceTarget::StaticStruct();
+}
+
+void FAnimInstanceTargetMenuProvider::PopulateTargetMenu(
+	FMenuBuilder& MenuBuilder,
+	TObjectPtr<UObject> BoundObject,
+	TFunction<void(TInstancedStruct<FMovieSceneMixedAnimationTarget>)> OnTargetSelected)
+{
+	MenuBuilder.AddMenuEntry(
+		LOCTEXT("CustomAnimInstanceTarget", "Custom Anim Instance"),
+		LOCTEXT("CustomAnimInstanceTargetTooltip", "Use a custom anim instance on the skeletal mesh component."),
+		FSlateIcon(),
+		FUIAction(
+			FExecuteAction::CreateLambda([OnTargetSelected]()
+			{
+				OnTargetSelected(TInstancedStruct<FMovieSceneAnimInstanceTarget>::Make());
+			})
+		)
+	);
+}
+
+//------------------------------------------------------------------------------
+// FAnimBlueprintTargetMenuProvider
+//------------------------------------------------------------------------------
+
+UScriptStruct* FAnimBlueprintTargetMenuProvider::GetHandledTargetStructType() const
+{
+	return FMovieSceneAnimBlueprintTarget::StaticStruct();
+}
+
+void FAnimBlueprintTargetMenuProvider::PopulateTargetMenu(
+	FMenuBuilder& MenuBuilder,
+	TObjectPtr<UObject> BoundObject,
+	TFunction<void(TInstancedStruct<FMovieSceneMixedAnimationTarget>)> OnTargetSelected)
+{
+	MenuBuilder.AddSubMenu(
+		LOCTEXT("AnimBlueprintTarget", "Anim Blueprint Target"),
+		LOCTEXT("AnimBlueprintTargetTooltip", "Target a specific Sequencer Mixer Target node in an Animation Blueprint."),
+		FNewMenuDelegate::CreateRaw(this, &FAnimBlueprintTargetMenuProvider::PopulateAnimBlueprintTargetSubmenu, BoundObject, OnTargetSelected)
+	);
+}
+
+void FAnimBlueprintTargetMenuProvider::PopulateAnimBlueprintTargetSubmenu(
+	FMenuBuilder& MenuBuilder,
+	TObjectPtr<UObject> BoundObject,
+	TFunction<void(TInstancedStruct<FMovieSceneMixedAnimationTarget>)> OnTargetSelected)
+{
+	// Try to find SequencerMixerTarget nodes from the bound object's anim blueprint
+	TArray<FName> TargetNames;
+
+	// Get the skeletal mesh component from the bound object
+	USkeletalMeshComponent* SkeletalMeshComponent = nullptr;
+	if (AActor* Actor = Cast<AActor>(BoundObject))
+	{
+		SkeletalMeshComponent = Actor->FindComponentByClass<USkeletalMeshComponent>();
+	}
+	else if (USkeletalMeshComponent* SMC = Cast<USkeletalMeshComponent>(BoundObject))
+	{
+		SkeletalMeshComponent = SMC;
+	}
+
+	if (SkeletalMeshComponent)
+	{
+		// Get the configured anim blueprint class
+		if (UClass* AnimClass = SkeletalMeshComponent->GetAnimClass())
+		{
+			if (UAnimBlueprintGeneratedClass* AnimBPClass = Cast<UAnimBlueprintGeneratedClass>(AnimClass))
+			{
+				// Find all SequencerMixerTarget nodes in the anim blueprint
+				// Use CDO to get default property values since we may not have a runtime instance
+				UAnimInstance* AnimCDO = AnimBPClass->GetDefaultObject<UAnimInstance>();
+				for (const FStructProperty* NodeProperty : AnimBPClass->AnimNodeProperties)
+				{
+					if (NodeProperty && NodeProperty->Struct == FAnimNode_SequencerMixerTarget::StaticStruct())
+					{
+						// Get the node and its target name from the CDO
+						const FAnimNode_SequencerMixerTarget* Node = NodeProperty->ContainerPtrToValuePtr<FAnimNode_SequencerMixerTarget>(AnimCDO);
+						if (Node && !Node->TargetName.IsNone())
+						{
+							TargetNames.AddUnique(Node->TargetName);
+						}
+					}
+				}
+			}
+		}
+	}
+
+	// Always add the default target name
+	TargetNames.AddUnique(FAnimNode_SequencerMixerTarget::DefaultTargetName);
+
+	// Add menu entries for each discovered target
+	for (const FName& TargetName : TargetNames)
+	{
+		MenuBuilder.AddMenuEntry(
+			FText::FromName(TargetName),
+			FText::Format(LOCTEXT("AnimBlueprintTargetEntryTooltip", "Target the '{0}' Sequencer Mixer Target node."), FText::FromName(TargetName)),
+			FSlateIcon(),
+			FUIAction(
+				FExecuteAction::CreateLambda([OnTargetSelected, TargetName]()
+				{
+					OnTargetSelected(TInstancedStruct<FMovieSceneAnimBlueprintTarget>::Make(TargetName));
+				})
+			)
+		);
+	}
+}
+
+//------------------------------------------------------------------------------
+// FAnimNextInjectionTargetMenuProvider
+//------------------------------------------------------------------------------
+
+UScriptStruct* FAnimNextInjectionTargetMenuProvider::GetHandledTargetStructType() const
+{
+	return FMovieSceneAnimNextInjectionTarget::StaticStruct();
+}
+
+void FAnimNextInjectionTargetMenuProvider::PopulateTargetMenu(
+	FMenuBuilder& MenuBuilder,
+	TObjectPtr<UObject> BoundObject,
+	TFunction<void(TInstancedStruct<FMovieSceneMixedAnimationTarget>)> OnTargetSelected)
+{
+	MenuBuilder.AddSubMenu(
+		LOCTEXT("UAFModuleInjectionTarget", "UAF Module Injection"),
+		LOCTEXT("UAFModuleInjectionTargetTooltip", "Target an injection site on an Unreal Animation Framework module."),
+		FNewMenuDelegate::CreateRaw(this, &FAnimNextInjectionTargetMenuProvider::PopulateAnimNextTargetSubmenu, BoundObject, OnTargetSelected)
+	);
+}
+
+void FAnimNextInjectionTargetMenuProvider::PopulateAnimNextTargetSubmenu(
+	FMenuBuilder& MenuBuilder,
+	TObjectPtr<UObject> BoundObject,
+	TFunction<void(TInstancedStruct<FMovieSceneMixedAnimationTarget>)> OnTargetSelected)
+{
+	using namespace UE::UAF::Editor;
+
+	// Always add a default entry that uses automatic injection site discovery
+	MenuBuilder.AddMenuEntry(
+		LOCTEXT("UAFDefaultInjectionSite", "Default"),
+		LOCTEXT("UAFDefaultInjectionSiteTooltip", "Use the default injection site on the actor's module."),
+		FSlateIcon(),
+		FUIAction(
+			FExecuteAction::CreateLambda([OnTargetSelected]()
+			{
+				OnTargetSelected(TInstancedStruct<FMovieSceneAnimNextInjectionTarget>::Make());
+			})
+		)
+	);
+
+	// Add an embedded variable picker for selecting a specific injection site
+	FVariablePickerArgs PickerArgs;
+	PickerArgs.OnVariablePicked = FOnVariablePicked::CreateLambda([OnTargetSelected](const FAnimNextSoftVariableReference& InVariableReference, const FAnimNextParamType& InType)
+	{
+		OnTargetSelected(TInstancedStruct<FMovieSceneAnimNextInjectionTarget>::Make(InVariableReference));
+	});
+
+	// Filter to only show FAnimNextAnimGraph types
+	FAnimNextParamType FilterType = FAnimNextParamType::FromString(TEXT("FAnimNextAnimGraph"));
+	PickerArgs.OnFilterVariableType = FOnFilterVariableType::CreateLambda([FilterType](const FAnimNextParamType& InParamType) -> EFilterVariableResult
+	{
+		if (FilterType.IsValid())
+		{
+			if (!UE::UAF::FParamUtils::GetCompatibility(FilterType, InParamType).IsCompatibleWithDataLoss())
+			{
+				return EFilterVariableResult::Exclude;
+			}
+		}
+		return EFilterVariableResult::Include;
+	});
+
+	MenuBuilder.AddWidget(
+		SNew(SVariablePickerCombo)
+		.PickerArgs(PickerArgs)
+		.VariableName(LOCTEXT("SelectInjectionSite", "Select Injection Site...")),
+		LOCTEXT("InjectionSitePickerLabel", ""),
+		/*bNoIndent=*/ true
+	);
+}
+
+#undef LOCTEXT_NAMESPACE
diff --git a/Engine/Plugins/Experimental/MovieSceneAnimMixer/Source/MovieSceneAnimMixerEditor/Private/AnimMixerTargetMenuProviders.h b/Engine/Plugins/Experimental/MovieSceneAnimMixer/Source/MovieSceneAnimMixerEditor/Private/AnimMixerTargetMenuProviders.h
new file mode 100644
--- /dev/null
+++ b/Engine/Plugins/Experimental/MovieSceneAnimMixer/Source/MovieSceneAnimMixerEditor/Private/AnimMixerTargetMenuProviders.h
@@ -0,0 +1,79 @@
+// Copyright Epic Games, Inc. All Rights Reserved.
+
+#pragma once
+
+#include "IMovieSceneAnimMixerTargetMenuProvider.h"
+
+/**
+ * Target menu provider for the Automatic target type (base FMovieSceneMixedAnimationTarget).
+ * Adds a simple "Automatic" menu entry.
+ */
+class FAutomaticTargetMenuProvider : public IMovieSceneAnimMixerTargetMenuProvider
+{
+public:
+	virtual UScriptStruct* GetHandledTargetStructType() const override;
+	virtual void PopulateTargetMenu(
+		FMenuBuilder& MenuBuilder,
+		TObjectPtr<UObject> BoundObject,
+		TFunction<void(TInstancedStruct<FMovieSceneMixedAnimationTarget>)> OnTargetSelected) override;
+	virtual int32 GetTargetMenuPriority() const override { return 0; }
+};
+
+/**
+ * Target menu provider for the Custom Anim Instance target type.
+ * Adds a simple "Custom Anim Instance" menu entry.
+ */
+class FAnimInstanceTargetMenuProvider : public IMovieSceneAnimMixerTargetMenuProvider
+{
+public:
+	virtual UScriptStruct* GetHandledTargetStructType() const override;
+	virtual void PopulateTargetMenu(
+		FMenuBuilder& MenuBuilder,
+		TObjectPtr<UObject> BoundObject,
+		TFunction<void(TInstancedStruct<FMovieSceneMixedAnimationTarget>)> OnTargetSelected) override;
+	virtual int32 GetTargetMenuPriority() const override { return 10; }
+};
+
+/**
+ * Target menu provider for the Anim Blueprint Target type.
+ * Queries the bound object's anim instance for SequencerMixerTarget nodes and
+ * creates a submenu with discovered target names, plus the default name.
+ */
+class FAnimBlueprintTargetMenuProvider : public IMovieSceneAnimMixerTargetMenuProvider
+{
+public:
+	virtual UScriptStruct* GetHandledTargetStructType() const override;
+	virtual void PopulateTargetMenu(
+		FMenuBuilder& MenuBuilder,
+		TObjectPtr<UObject> BoundObject,
+		TFunction<void(TInstancedStruct<FMovieSceneMixedAnimationTarget>)> OnTargetSelected) override;
+	virtual int32 GetTargetMenuPriority() const override { return 20; }
+
+private:
+	void PopulateAnimBlueprintTargetSubmenu(
+		FMenuBuilder& MenuBuilder,
+		TObjectPtr<UObject> BoundObject,
+		TFunction<void(TInstancedStruct<FMovieSceneMixedAnimationTarget>)> OnTargetSelected);
+};
+
+/**
+ * Target menu provider for the UAF Module Injection target type.
+ * Queries the bound object's AnimNextComponent for injection sites and
+ * creates a submenu with discovered injection sites, plus the default one.
+ */
+class FAnimNextInjectionTargetMenuProvider : public IMovieSceneAnimMixerTargetMenuProvider
+{
+public:
+	virtual UScriptStruct* GetHandledTargetStructType() const override;
+	virtual void PopulateTargetMenu(
+		FMenuBuilder& MenuBuilder,
+		TObjectPtr<UObject> BoundObject,
+		TFunction<void(TInstancedStruct<FMovieSceneMixedAnimationTarget>)> OnTargetSelected) override;
+	virtual int32 GetTargetMenuPriority() const override { return 30; }
+
+private:
+	void PopulateAnimNextTargetSubmenu(
+		FMenuBuilder& MenuBuilder,
+		TObjectPtr<UObject> BoundObject,
+		TFunction<void(TInstancedStruct<FMovieSceneMixedAnimationTarget>)> OnTargetSelected);
+};
diff --git a/Engine/Plugins/Experimental/MovieSceneAnimMixer/Source/MovieSceneAnimMixerEditor/Private/MovieSceneAnimMixerEditorModule.cpp b/Engine/Plugins/Experimental/MovieSceneAnimMixer/Source/MovieSceneAnimMixerEditor/Private/MovieSceneAnimMixerEditorModule.cpp
--- a/Engine/Plugins/Experimental/MovieSceneAnimMixer/Source/MovieSceneAnimMixerEditor/Private/MovieSceneAnimMixerEditorModule.cpp
+++ b/Engine/Plugins/Experimental/MovieSceneAnimMixer/Source/MovieSceneAnimMixerEditor/Private/MovieSceneAnimMixerEditorModule.cpp
@@ -3,7 +3,9 @@
 #include "MovieSceneAnimMixerEditorModule.h"
 
 #include "AnimBlueprintExtension_SequencerMixerTarget.h"
+#include "AnimMixerTargetMenuProviders.h"
 #include "IAnimGraph_SequencerMixerTargetConnector.h"
+#include "IMovieSceneAnimMixerTargetMenuProvider.h"
 #include "ISequencerModule.h"
 #include "Modules/ModuleInterface.h"
 #include "Modules/ModuleManager.h"
@@ -104,10 +106,39 @@ namespace UE::MovieScene
 		));
 
 		SequencerModule.RegisterChannelInterface<FMovieSceneByteChannelDefaultOnly>();
+
+		// Register target menu providers
+		AutomaticTargetMenuProvider = MakeUnique<FAutomaticTargetMenuProvider>();
+		AnimInstanceTargetMenuProvider = MakeUnique<FAnimInstanceTargetMenuProvider>();
+		AnimBlueprintTargetMenuProvider = MakeUnique<FAnimBlueprintTargetMenuProvider>();
+		AnimNextInjectionTargetMenuProvider = MakeUnique<FAnimNextInjectionTargetMenuProvider>();
+
+		IModularFeatures::Get().RegisterModularFeature(IMovieSceneAnimMixerTargetMenuProvider::GetModularFeatureName(), AutomaticTargetMenuProvider.Get());
+		IModularFeatures::Get().RegisterModularFeature(IMovieSceneAnimMixerTargetMenuProvider::GetModularFeatureName(), AnimInstanceTargetMenuProvider.Get());
+		IModularFeatures::Get().RegisterModularFeature(IMovieSceneAnimMixerTargetMenuProvider::GetModularFeatureName(), AnimBlueprintTargetMenuProvider.Get());
+		IModularFeatures::Get().RegisterModularFeature(IMovieSceneAnimMixerTargetMenuProvider::GetModularFeatureName(), AnimNextInjectionTargetMenuProvider.Get());
 	}
 
 	void FMovieSceneAnimMixerEditorModule::ShutdownModule()
 	{
+		// Unregister target menu providers
+		if (AutomaticTargetMenuProvider)
+		{
+			IModularFeatures::Get().UnregisterModularFeature(IMovieSceneAnimMixerTargetMenuProvider::GetModularFeatureName(), AutomaticTargetMenuProvider.Get());
+		}
+		if (AnimInstanceTargetMenuProvider)
+		{
+			IModularFeatures::Get().UnregisterModularFeature(IMovieSceneAnimMixerTargetMenuProvider::GetModularFeatureName(), AnimInstanceTargetMenuProvider.Get());
+		}
+		if (AnimBlueprintTargetMenuProvider)
+		{
+			IModularFeatures::Get().UnregisterModularFeature(IMovieSceneAnimMixerTargetMenuProvider::GetModularFeatureName(), AnimBlueprintTargetMenuProvider.Get());
+		}
+		if (AnimNextInjectionTargetMenuProvider)
+		{
+			IModularFeatures::Get().UnregisterModularFeature(IMovieSceneAnimMixerTargetMenuProvider::GetModularFeatureName(), AnimNextInjectionTargetMenuProvider.Get());
+		}
+
 		if (ISequencerModule* SequencerModule = FModuleManager::Get().GetModulePtr<ISequencerModule>(SequencerModuleName.Resolve()))
 		{
 			SequencerModule->UnregisterTrackModel(AnimationMixerTrackModelHandle);
diff --git a/Engine/Plugins/Experimental/MovieSceneAnimMixer/Source/MovieSceneAnimMixerEditor/Private/MovieSceneAnimMixerEditorModule.h b/Engine/Plugins/Experimental/MovieSceneAnimMixer/Source/MovieSceneAnimMixerEditor/Private/MovieSceneAnimMixerEditorModule.h
--- a/Engine/Plugins/Experimental/MovieSceneAnimMixer/Source/MovieSceneAnimMixerEditor/Private/MovieSceneAnimMixerEditorModule.h
+++ b/Engine/Plugins/Experimental/MovieSceneAnimMixer/Source/MovieSceneAnimMixerEditor/Private/MovieSceneAnimMixerEditorModule.h
@@ -6,21 +6,32 @@
 #include "IAnimGraph_SequencerMixerTargetConnector.h"
 #include "Modules/ModuleInterface.h"
 
+class FAutomaticTargetMenuProvider;
+class FAnimInstanceTargetMenuProvider;
+class FAnimBlueprintTargetMenuProvider;
+class FAnimNextInjectionTargetMenuProvider;
+
 namespace UE::MovieScene
 {
 	class FMovieSceneAnimMixerEditorModule : public IModuleInterface, public IAnimGraph_SequencerMixerTargetConnector
 	{
 	public:
-		
+
 		// IModuleInterface interface
 		virtual void StartupModule() override;
 		virtual void ShutdownModule() override;
-		
+
 		virtual void GetSequencerMixerTargetRequiredExtensions(TArray<TSubclassOf<UAnimBlueprintExtension>>& OutExtensions) const override;
 
 	private:
 		FDelegateHandle AnimationTrackEditorHandle;
 		FDelegateHandle AnimationMixerTrackModelHandle;
+
+		// Target menu providers
+		TUniquePtr<FAutomaticTargetMenuProvider> AutomaticTargetMenuProvider;
+		TUniquePtr<FAnimInstanceTargetMenuProvider> AnimInstanceTargetMenuProvider;
+		TUniquePtr<FAnimBlueprintTargetMenuProvider> AnimBlueprintTargetMenuProvider;
+		TUniquePtr<FAnimNextInjectionTargetMenuProvider> AnimNextInjectionTargetMenuProvider;
 	};
 }
 
diff --git a/Engine/Plugins/Experimental/MovieSceneAnimMixer/Source/MovieSceneAnimMixerEditor/Private/MovieSceneAnimationMixerTrackEditor.cpp b/Engine/Plugins/Experimental/MovieSceneAnimMixer/Source/MovieSceneAnimMixerEditor/Private/MovieSceneAnimationMixerTrackEditor.cpp
--- a/Engine/Plugins/Experimental/MovieSceneAnimMixer/Source/MovieSceneAnimMixerEditor/Private/MovieSceneAnimationMixerTrackEditor.cpp
+++ b/Engine/Plugins/Experimental/MovieSceneAnimMixer/Source/MovieSceneAnimMixerEditor/Private/MovieSceneAnimationMixerTrackEditor.cpp
@@ -6,11 +6,14 @@
 #include "AssetRegistry/AssetRegistryModule.h"
 #include "ClassIconFinder.h"
 #include "Components/SkeletalMeshComponent.h"
+#include "Features/IModularFeatures.h"
 #include "Framework/MultiBox/MultiBoxBuilder.h"
 #include "GameFramework/Actor.h"
 #include "IMovieSceneAnimMixerItemMenuProvider.h"
+#include "IMovieSceneAnimMixerTargetMenuProvider.h"
 #include "ISequencer.h"
 #include "MovieSceneAnimationMixerItemInterface.h"
+#include "MovieSceneAnimationMixerLayer.h"
 #include "MovieSceneAnimationMixerTrack.h"
 #include "MovieSceneTrack.h"
 #include "MVVM/Extensions/IObjectBindingExtension.h"
@@ -24,8 +27,8 @@
 #include "ScopedTransaction.h"
 #include "Sections/MovieSceneSkeletalAnimationSection.h"
 #include "SequencerUtilities.h"
-#include "MovieSceneAnimationMixerLayer.h"
-#include "Features/IModularFeatures.h"
+#include "Widgets/Input/SComboButton.h"
+#include "Widgets/Images/SImage.h"
 
 #define LOCTEXT_NAMESPACE "MovieSceneAnimationMixerTrackEditor"
 
@@ -33,6 +36,152 @@
 namespace UE::Sequencer
 {
 
+/**
+ * Combo widget for selecting the animation mixer target type in the outliner.
+ * Displays the current target icon and short name, and provides a dropdown menu
+ * to change the target using registered IMovieSceneAnimMixerTargetMenuProvider instances.
+ */
+class SAnimMixerTargetComboWidget : public SCompoundWidget
+{
+public:
+	SLATE_BEGIN_ARGS(SAnimMixerTargetComboWidget) {}
+	SLATE_END_ARGS()
+
+	void Construct(const FArguments& InArgs, TWeakObjectPtr<UMovieSceneAnimationMixerTrack> InWeakTrack, TWeakPtr<ISequencer> InWeakSequencer, FGuid InObjectBinding)
+	{
+		WeakTrack = InWeakTrack;
+		WeakSequencer = InWeakSequencer;
+		ObjectBinding = InObjectBinding;
+
+		ChildSlot
+		.HAlign(HAlign_Left)
+		[
+			SNew(SBox)
+			.WidthOverride(120.f)
+			.Clipping(EWidgetClipping::ClipToBounds)
+			[
+				SNew(SComboButton)
+				.ToolTipText(this, &SAnimMixerTargetComboWidget::GetTargetDisplayName)
+				.OnGetMenuContent(this, &SAnimMixerTargetComboWidget::GetMenuContent)
+				.ButtonContent()
+				[
+					SNew(SHorizontalBox)
+					+ SHorizontalBox::Slot()
+					.AutoWidth()
+					.VAlign(VAlign_Center)
+					.Padding(0.f, 0.f, 4.f, 0.f)
+					[
+						SNew(SImage)
+						.Image(this, &SAnimMixerTargetComboWidget::GetTargetIcon)
+						.Visibility(this, &SAnimMixerTargetComboWidget::GetTargetIconVisibility)
+					]
+					+ SHorizontalBox::Slot()
+					.VAlign(VAlign_Center)
+					[
+						SNew(STextBlock)
+						.Text(this, &SAnimMixerTargetComboWidget::GetTargetDisplayName)
+						.OverflowPolicy(ETextOverflowPolicy::Ellipsis)
+					]
+				]
+			]
+		];
+	}
+
+private:
+	TSharedRef<SWidget> GetMenuContent()
+	{
+		FMenuBuilder MenuBuilder(true, nullptr);
+
+		// Get the bound object for target menu providers
+		UObject* BoundObject = nullptr;
+		TSharedPtr<ISequencer> Sequencer = WeakSequencer.Pin();
+		if (Sequencer && ObjectBinding.IsValid())
+		{
+			TArrayView<TWeakObjectPtr<>> BoundObjects = Sequencer->FindBoundObjects(ObjectBinding, Sequencer->GetFocusedTemplateID());
+			if (BoundObjects.Num() > 0)
+			{
+				BoundObject = BoundObjects[0].Get();
+			}
+		}
+
+		// Find all registered target menu providers
+		TArray<IMovieSceneAnimMixerTargetMenuProvider*> Providers = IModularFeatures::Get().GetModularFeatureImplementations<IMovieSceneAnimMixerTargetMenuProvider>(IMovieSceneAnimMixerTargetMenuProvider::GetModularFeatureName());
+
+		// Sort providers by priority (lower values appear first)
+		Providers.Sort([](const IMovieSceneAnimMixerTargetMenuProvider& A, const IMovieSceneAnimMixerTargetMenuProvider& B)
+		{
+			return A.GetTargetMenuPriority() < B.GetTargetMenuPriority();
+		});
+
+		// Callback when a target is selected
+		auto OnTargetSelected = [this](TInstancedStruct<FMovieSceneMixedAnimationTarget> Target)
+		{
+			SetTarget(MoveTemp(Target));
+		};
+
+		// Populate menu from all providers
+		for (IMovieSceneAnimMixerTargetMenuProvider* Provider : Providers)
+		{
+			if (Provider)
+			{
+				Provider->PopulateTargetMenu(MenuBuilder, BoundObject, OnTargetSelected);
+			}
+		}
+
+		return MenuBuilder.MakeWidget();
+	}
+
+	const FSlateBrush* GetTargetIcon() const
+	{
+		UMovieSceneAnimationMixerTrack* Track = WeakTrack.Get();
+		if (Track)
+		{
+			const FMovieSceneMixedAnimationTarget* Target = Track->MixedAnimationTarget.GetPtr<FMovieSceneMixedAnimationTarget>();
+			if (Target)
+			{
+				return Target->GetIcon();
+			}
+		}
+		return nullptr;
+	}
+
+	EVisibility GetTargetIconVisibility() const
+	{
+		return GetTargetIcon() != nullptr ? EVisibility::Visible : EVisibility::Collapsed;
+	}
+
+	FText GetTargetDisplayName() const
+	{
+		UMovieSceneAnimationMixerTrack* Track = WeakTrack.Get();
+		if (Track)
+		{
+			const FMovieSceneMixedAnimationTarget* Target = Track->MixedAnimationTarget.GetPtr<FMovieSceneMixedAnimationTarget>();
+			if (Target)
+			{
+				return Target->GetShortDisplayName();
+			}
+		}
+		return LOCTEXT("UnknownTarget", "Unknown");
+	}
+
+	void SetTarget(TInstancedStruct<FMovieSceneMixedAnimationTarget> NewTarget)
+	{
+		UMovieSceneAnimationMixerTrack* Track = WeakTrack.Get();
+		TSharedPtr<ISequencer> Sequencer = WeakSequencer.Pin();
+		if (Track && Sequencer)
+		{
+			FScopedTransaction Transaction(LOCTEXT("SetAnimMixerTarget", "Set Animation Mixer Target"));
+			Track->Modify();
+			Track->MixedAnimationTarget = MoveTemp(NewTarget);
+			Sequencer->NotifyMovieSceneDataChanged(EMovieSceneDataChangeType::TrackValueChanged);
+		}
+	}
+
+	TWeakObjectPtr<UMovieSceneAnimationMixerTrack> WeakTrack;
+	TWeakPtr<ISequencer> WeakSequencer;
+	FGuid ObjectBinding;
+};
+
 class FAnimMixerAnimationSection
 	: public FCommonAnimationSection
 {
@@ -273,17 +422,88 @@ void FAnimationMixerTrackEditor::PopulateObjectBindingAnimationMenu(FMenuBuilder
 {
 	if (ObjectClass != nullptr && (ObjectClass->IsChildOf(USkeletalMeshComponent::StaticClass()) || ObjectClass->IsChildOf(AActor::StaticClass())))
 	{
-		MenuBuilder.AddMenuEntry(
+		MenuBuilder.AddSubMenu(
 			LOCTEXT("AddAnimationTrack", "Animation Mixer"),
 			LOCTEXT("AddAnimationTrackTooltip", "Adds a new animation track for playing back Anim Sequences and other sources of animation."),
-			FSlateIcon(FAppStyle::GetAppStyleSetName(), "Sequencer.Tracks.Animation"),
-			FUIAction(
-				FExecuteAction::CreateSP(this, &FAnimationMixerTrackEditor::HandleAddAnimationTrackMenuEntryExecute, ObjectBindings)
-			)
+			FNewMenuDelegate::CreateSP(this, &FAnimationMixerTrackEditor::BuildTargetSelectionMenu, TArray<FGuid>(ObjectBindings), ObjectClass),
+			FUIAction(),
+			NAME_None,
+			EUserInterfaceActionType::None,
+			false,
+			FSlateIcon(FAppStyle::GetAppStyleSetName(), "Sequencer.Tracks.Animation")
 		);
 	}
 }
 
+void FAnimationMixerTrackEditor::BuildTargetSelectionMenu(FMenuBuilder& MenuBuilder, TArray<FGuid> ObjectBindings, const UClass* ObjectClass)
+{
+	// Get the bound object for target menu providers
+	UObject* BoundObject = nullptr;
+	TSharedPtr<ISequencer> Sequencer = GetSequencer();
+	if (Sequencer && ObjectBindings.Num() > 0)
+	{
+		TArrayView<TWeakObjectPtr<>> BoundObjects = Sequencer->FindBoundObjects(ObjectBindings[0], Sequencer->GetFocusedTemplateID());
+		if (BoundObjects.Num() > 0)
+		{
+			BoundObject = BoundObjects[0].Get();
+		}
+	}
+
+	// Find all registered target menu providers
+	TArray<IMovieSceneAnimMixerTargetMenuProvider*> Providers = IModularFeatures::Get().GetModularFeatureImplementations<IMovieSceneAnimMixerTargetMenuProvider>(IMovieSceneAnimMixerTargetMenuProvider::GetModularFeatureName());
+
+	// Sort providers by priority (lower values appear first)
+	Providers.Sort([](const IMovieSceneAnimMixerTargetMenuProvider& A, const IMovieSceneAnimMixerTargetMenuProvider& B)
+	{
+		return A.GetTargetMenuPriority() < B.GetTargetMenuPriority();
+	});
+
+	// Callback when a target is selected
+	auto OnTargetSelected = [this, ObjectBindings](TInstancedStruct<FMovieSceneMixedAnimationTarget> Target)
+	{
+		HandleAddMixerTrackWithTarget(ObjectBindings, MoveTemp(Target));
+	};
+
+	// Populate menu from all providers
+	for (IMovieSceneAnimMixerTargetMenuProvider* Provider : Providers)
+	{
+		if (Provider)
+		{
+			Provider->PopulateTargetMenu(MenuBuilder, BoundObject, OnTargetSelected);
+		}
+	}
+}
+
+void FAnimationMixerTrackEditor::HandleAddMixerTrackWithTarget(TArray<FGuid> ObjectBindings, TInstancedStruct<FMovieSceneMixedAnimationTarget> Target)
+{
+	UMovieScene* MovieScene = GetFocusedMovieScene();
+
+	if (MovieScene == nullptr)
+	{
+		return;
+	}
+
+	if (MovieScene->IsReadOnly())
+	{
+		return;
+	}
+
+	const FScopedTransaction Transaction(LOCTEXT("AddMixerTrack_Transaction", "Add Mixer Track"));
+
+	MovieScene->Modify();
+
+	for (const FGuid& Guid : ObjectBindings)
+	{
+		FFindOrCreateTrackResult TrackResult = FindOrCreateTrackForObject(Guid, UMovieSceneAnimationMixerTrack::StaticClass());
+		if (TrackResult.bWasCreated)
+		{
+			UMovieSceneAnimationMixerTrack* MixerTrack = CastChecked<UMovieSceneAnimationMixerTrack>(TrackResult.Track);
+			MixerTrack->SetDisplayName(MixerTrack->GetDefaultDisplayName());
+			MixerTrack->MixedAnimationTarget = Target;
+		}
+	}
+}
+
 TSharedPtr<SWidget> FAnimationMixerTrackEditor::BuildOutlinerColumnWidget(const FBuildColumnWidgetParams& Params, const FName& ColumnName)
 {
 	using namespace UE::Sequencer;
@@ -324,6 +544,25 @@ TSharedPtr<SWidget> FAnimationMixerTrackEditor::BuildOutlinerColumnWidget(const
 			Params.ViewModel);
 	}
 
+	if (ColumnName == FCommonOutlinerNames::Edit)
+	{
+		// Only show target combo on the track itself, not on layers
+		if (TViewModelPtr<ITrackExtension> TrackModel = Params.ViewModel.ImplicitCast())
+		{
+			if (UMovieSceneAnimationMixerTrack* MixerTrack = Cast<UMovieSceneAnimationMixerTrack>(TrackModel->GetTrack()))
+			{
+				// Get the object binding for this track
+				FGuid ObjectBindingGuid;
+				if (TViewModelPtr<IObjectBindingExtension> ObjectBinding = Params.ViewModel->FindAncestorOfType<IObjectBindingExtension>())
+				{
+					ObjectBindingGuid = ObjectBinding->GetObjectGuid();
+				}
+
+				return SNew(SAnimMixerTargetComboWidget, MixerTrack, GetSequencer(), ObjectBindingGuid);
+			}
+		}
+	}
+
 	return FMovieSceneTrackEditor::BuildOutlinerColumnWidget(Params, ColumnName);
 }
 
diff --git a/Engine/Plugins/Experimental/MovieSceneAnimMixer/Source/MovieSceneAnimMixerEditor/Private/MovieSceneAnimationMixerTrackEditor.h b/Engine/Plugins/Experimental/MovieSceneAnimMixer/Source/MovieSceneAnimMixerEditor/Private/MovieSceneAnimationMixerTrackEditor.h
--- a/Engine/Plugins/Experimental/MovieSceneAnimMixer/Source/MovieSceneAnimMixerEditor/Private/MovieSceneAnimationMixerTrackEditor.h
+++ b/Engine/Plugins/Experimental/MovieSceneAnimMixer/Source/MovieSceneAnimMixerEditor/Private/MovieSceneAnimationMixerTrackEditor.h
@@ -5,6 +5,8 @@
 #include "Templates/SharedPointer.h"
 #include "Templates/SubclassOf.h"
 #include "MovieSceneTrackEditor.h"
+#include "MovieSceneMixedAnimationTarget.h"
+#include "StructUtils/InstancedStruct.h"
 #include "TrackEditors/CommonAnimationTrackEditor.h"
 #include "IMovieSceneAnimMixerItemMenuProvider.h"
 
@@ -60,6 +62,8 @@ private:
 
 	TSharedRef<SWidget> BuildAddSectionSubMenu(TWeakViewModelPtr<IOutlinerExtension> WeakViewModel, TWeakViewModelPtr<FSequencerEditorViewModel> WeakEditor);
 	void HandleAddAnimationTrackMenuEntryExecute(TArray<FGuid> ObjectBindings);
+	void HandleAddMixerTrackWithTarget(TArray<FGuid> ObjectBindings, TInstancedStruct<FMovieSceneMixedAnimationTarget> Target);
+	void BuildTargetSelectionMenu(FMenuBuilder& MenuBuilder, TArray<FGuid> ObjectBindings, const UClass* ObjectClass);
 	void PopulateAddAnimationMenu(FMenuBuilder& MenuBuilder, TWeakViewModelPtr<IOutlinerExtension> WeakViewModel, TWeakViewModelPtr<FSequencerEditorViewModel> WeakEditor);
 
 	bool FilterAnimSequences(const FAssetData& AssetData, USkeleton* Skeleton);
diff --git a/Engine/Plugins/Experimental/UAF/UAF/Source/UAFEditor/Private/Variables/SGraphPinVariableReference.cpp b/Engine/Plugins/Experimental/UAF/UAF/Source/UAFEditor/Private/Variables/SGraphPinVariableReference.cpp
--- a/Engine/Plugins/Experimental/UAF/UAF/Source/UAFEditor/Private/Variables/SGraphPinVariableReference.cpp
+++ b/Engine/Plugins/Experimental/UAF/UAF/Source/UAFEditor/Private/Variables/SGraphPinVariableReference.cpp
@@ -3,7 +3,7 @@
 #include "SGraphPinVariableReference.h"
 #include "EditorUtils.h"
 #include "K2Node_CallFunction.h"
-#include "SVariablePickerCombo.h"
+#include "Variables/SVariablePickerCombo.h"
 #include "ScopedTransaction.h"
 #include "UncookedOnlyUtils.h"
 #include "Component/AnimNextComponent.h"
diff --git a/Engine/Plugins/Experimental/UAF/UAF/Source/UAFEditor/Private/Variables/SVariablePickerCombo.cpp b/Engine/Plugins/Experimental/UAF/UAF/Source/UAFEditor/Private/Variables/SVariablePickerCombo.cpp
--- a/Engine/Plugins/Experimental/UAF/UAF/Source/UAFEditor/Private/Variables/SVariablePickerCombo.cpp
+++ b/Engine/Plugins/Experimental/UAF/UAF/Source/UAFEditor/Private/Variables/SVariablePickerCombo.cpp
@@ -1,6 +1,6 @@
 ﻿// Copyright Epic Games, Inc. All Rights Reserved.
 
-#include "SVariablePickerCombo.h"
+#include "Variables/SVariablePickerCombo.h"
 #include "SVariablePicker.h"
 #include "UncookedOnlyUtils.h"
 #include "Kismet2/BlueprintEditorUtils.h"
diff --git a/Engine/Plugins/Experimental/UAF/UAF/Source/UAFEditor/Private/Variables/SVariablePickerCombo.h b/Engine/Plugins/Experimental/UAF/UAF/Source/UAFEditor/Private/Variables/SVariablePickerCombo.h
deleted file mode 100644
diff --git a/Engine/Plugins/Experimental/UAF/UAF/Source/UAFEditor/Private/Variables/VariableReferencePropertyCustomization.cpp b/Engine/Plugins/Experimental/UAF/UAF/Source/UAFEditor/Private/Variables/VariableReferencePropertyCustomization.cpp
--- a/Engine/Plugins/Experimental/UAF/UAF/Source/UAFEditor/Private/Variables/VariableReferencePropertyCustomization.cpp
+++ b/Engine/Plugins/Experimental/UAF/UAF/Source/UAFEditor/Private/Variables/VariableReferencePropertyCustomization.cpp
@@ -5,7 +5,7 @@
 #include "PropertyHandle.h"
 #include "DetailWidgetRow.h"
 #include "IPropertyUtilities.h"
-#include "SVariablePickerCombo.h"
+#include "Variables/SVariablePickerCombo.h"
 #include "UncookedOnlyUtils.h"
 #include "Param/ParamCompatibility.h"
 #include "Param/ParamUtils.h"
diff --git a/Engine/Plugins/Experimental/UAF/UAF/Source/UAFEditor/Public/Variables/SVariablePickerCombo.h b/Engine/Plugins/Experimental/UAF/UAF/Source/UAFEditor/Public/Variables/SVariablePickerCombo.h
new file mode 100644
--- /dev/null
+++ b/Engine/Plugins/Experimental/UAF/UAF/Source/UAFEditor/Public/Variables/SVariablePickerCombo.h
@@ -0,0 +1,99 @@
+﻿// Copyright Epic Games, Inc. All Rights Reserved.
+
+#pragma once
+
+#include "Widgets/SCompoundWidget.h"
+#include "Variables/VariablePickerArgs.h"
+#include "EdGraph/EdGraphPin.h"
+#include "Param/ParamType.h"
+#include "Variables/AnimNextSoftVariableReference.h"
+
+struct FAnimNextParamInstanceIdentifier;
+
+namespace UE::UAF::Editor
+{
+	class SVariablePicker;
+	class SGraphPinVariableReference;
+}
+
+namespace UE::UAF::Editor
+{
+
+/** Retrieves the variable reference to display */
+using FOnGetVariableReference = TDelegate<FAnimNextSoftVariableReference()>;
+
+/** Retrieves the variable type to display */
+using FOnGetVariableType = TDelegate<FAnimNextParamType()>;
+
+class UAFEDITOR_API SVariablePickerCombo : public SCompoundWidget
+{
+public:
+	SLATE_BEGIN_ARGS(SVariablePickerCombo) {}
+
+	/** Arguments for configuring the picker in the dropdown menu */
+	SLATE_ARGUMENT(FVariablePickerArgs, PickerArgs)
+
+	/** Retrieves the variable name to display */
+	SLATE_EVENT(FOnGetVariableReference, OnGetVariableReference)
+
+	/** Retrieves the variable type to display */
+	SLATE_EVENT(FOnGetVariableType, OnGetVariableType)
+
+	/** The variable name text to display. This will be used if valid, otherwise the text of the variable reference will be used. */
+	SLATE_ATTRIBUTE(FText, VariableName)
+
+	/** The variable tooltip text to display. This will be used if valid, otherwise the tooltip of the variable reference will be used. */
+	SLATE_ATTRIBUTE(FText, VariableTooltip)
+
+	SLATE_END_ARGS()
+
+	void Construct(const FArguments& InArgs);
+
+private:
+	void RequestRefresh();
+
+	// Retrieves the variable reference to display
+	FOnGetVariableReference OnGetVariableReferenceDelegate;
+
+	// Retrieves the variable type to display
+	FOnGetVariableType OnGetVariableTypeDelegate;
+
+	// Cached pin type
+	FEdGraphPinType PinType;
+
+	// Cached variable reference
+	FAnimNextSoftVariableReference CachedVariableReference;
+
+	// Attribute for variable name
+	TAttribute<FText> VariableName;
+	
+	// Attribute for variable tooltip
+	TAttribute<FText> VariableTooltip;
+
+	// Cached display name
+	FText CachedVariableNameText;
+
+	// Cached tooltip
+	FText CachedVariableNameTooltipText;
+
+	// Cached variable type
+	FAnimNextParamType VariableType;
+
+	// Cached icon
+	const FSlateBrush* Icon = nullptr;
+
+	// Cached color
+	FSlateColor IconColor = FLinearColor::Gray;
+
+	// Picker widget used to focus after the popup is displayed
+	TWeakPtr<SVariablePicker> PickerWidget;
+
+	// Arguments for the picker popup
+	FVariablePickerArgs PickerArgs;
+
+	bool bRefreshRequested = false;
+
+	friend class SGraphPinVariableReference;
+};
+
+}
\ No newline at end of file
diff --git a/Engine/Source/Editor/MovieSceneTools/Private/IMovieSceneAnimMixerTargetMenuProvider.cpp b/Engine/Source/Editor/MovieSceneTools/Private/IMovieSceneAnimMixerTargetMenuProvider.cpp
new file mode 100644
--- /dev/null
+++ b/Engine/Source/Editor/MovieSceneTools/Private/IMovieSceneAnimMixerTargetMenuProvider.cpp
@@ -0,0 +1,9 @@
+// Copyright Epic Games, Inc. All Rights Reserved.
+
+#include "IMovieSceneAnimMixerTargetMenuProvider.h"
+
+FName IMovieSceneAnimMixerTargetMenuProvider::GetModularFeatureName()
+{
+	static FName FeatureName(TEXT("MovieSceneAnimMixerTargetProvider"));
+	return FeatureName;
+}
diff --git a/Engine/Source/Editor/MovieSceneTools/Public/IMovieSceneAnimMixerTargetMenuProvider.h b/Engine/Source/Editor/MovieSceneTools/Public/IMovieSceneAnimMixerTargetMenuProvider.h
new file mode 100644
--- /dev/null
+++ b/Engine/Source/Editor/MovieSceneTools/Public/IMovieSceneAnimMixerTargetMenuProvider.h
@@ -0,0 +1,54 @@
+// Copyright Epic Games, Inc. All Rights Reserved.
+
+#pragma once
+
+#include "CoreMinimal.h"
+#include "Features/IModularFeature.h"
+#include "MovieSceneMixedAnimationTarget.h"
+#include "StructUtils/InstancedStruct.h"
+
+class FMenuBuilder;
+
+/**
+ * Modular feature interface for providing custom UI when selecting animation mixer targets.
+ * Implement this interface to provide custom menu entries for target selection when creating
+ * Animation Mixer tracks or changing the target on an existing track.
+ *
+ * Each target type (e.g. Automatic, Custom Anim Instance, Anim Blueprint Target, UAF Module Injection)
+ * should have a corresponding provider that handles its menu population.
+ */
+class IMovieSceneAnimMixerTargetMenuProvider : public IModularFeature
+{
+public:
+	MOVIESCENETOOLS_API static FName GetModularFeatureName();
+
+	virtual ~IMovieSceneAnimMixerTargetMenuProvider() {}
+
+	/** Returns the target struct type this provider handles (e.g., FMovieSceneMixedAnimationTarget::StaticStruct()) */
+	virtual UScriptStruct* GetHandledTargetStructType() const = 0;
+
+	/**
+	 * Populate the menu with entries for this target type.
+	 * Called when building the target selection menu (both for new track creation and target change).
+	 *
+	 * For simple targets (like Automatic or Custom Anim Instance), this may add a single menu entry.
+	 * For complex targets (like Anim Blueprint Target), this may add a submenu with dynamically
+	 * discovered options based on the bound object.
+	 *
+	 * OnTargetSelected should be invoked by the menu provider with a constructed target once a target is selected.
+	 */
+	virtual void PopulateTargetMenu(
+		FMenuBuilder& MenuBuilder,
+		TObjectPtr<UObject> BoundObject,
+		TFunction<void(TInstancedStruct<FMovieSceneMixedAnimationTarget>)> OnTargetSelected) = 0;
+
+	/**
+	 * Returns the menu priority for this provider. Lower values appear first in the menu.
+	 *   0 - Automatic (default target)
+	 *   10 - Custom Anim Instance
+	 *   20 - Anim Blueprint Target
+	 *   30 - UAF Module Injection
+	 *   100+ - Custom/plugin targets
+	 */
+	virtual int32 GetTargetMenuPriority() const { return 100; }
+};
diff --git a/Engine/Source/Runtime/MovieSceneTracks/Public/MovieSceneMixedAnimationTarget.h b/Engine/Source/Runtime/MovieSceneTracks/Public/MovieSceneMixedAnimationTarget.h
--- a/Engine/Source/Runtime/MovieSceneTracks/Public/MovieSceneMixedAnimationTarget.h
+++ b/Engine/Source/Runtime/MovieSceneTracks/Public/MovieSceneMixedAnimationTarget.h
@@ -3,8 +3,11 @@
 #pragma once
 
 #include "CoreTypes.h"
+#include "Internationalization/Text.h"
 #include "MovieSceneMixedAnimationTarget.generated.h"
 
+struct FSlateBrush;
+
 /**
  * Defines a Mixed Anim Target that animation pose-producing tracks can send pose-producing tasks to. 
  * Multiple pose producing tasks sent to the same target will be chained and mixed based on priority, pose weight, masks, and sequence hierarchy.
@@ -15,7 +18,7 @@
 USTRUCT(meta=(DisplayName="Automatic"))
 struct FMovieSceneMixedAnimationTarget
 {
-	GENERATED_BODY() 
+	GENERATED_BODY()
 
 	virtual ~FMovieSceneMixedAnimationTarget() = default;
 
@@ -28,6 +31,14 @@ struct FMovieSceneMixedAnimationTarget
 
 	void SetHasFiredWarningForTarget(bool bNewHasFiredState) { bHasFiredWarningForTarget = bNewHasFiredState; }
 
+#if WITH_EDITOR
+	/** Get the icon to display for this target type in the Sequencer outliner. Returns nullptr for no icon. */
+	virtual const FSlateBrush* GetIcon() const { return nullptr; }
+
+	/** Get a short display name for this target (e.g., "Automatic", specific target node name, etc.) */
+	virtual FText GetShortDisplayName() const { return NSLOCTEXT("MovieSceneMixedAnimationTarget", "Automatic", "Automatic"); }
+#endif
+
 private:
 	bool bHasFiredWarningForTarget = false;
 };
