Sequencer: Anim Mixer Insertion changes.
* Allow insertion between layers when dragging sections only when CTRL is held.
* Add 'Add Layer Above' and 'Add Layer Below' menu items for Layer rows
* Slight code cleanup to add insert layer function to the anim mixer track

#Synced-CL 49155801

diff --git a/Engine/Plugins/Experimental/MovieSceneAnimMixer/Source/MovieSceneAnimMixer/Private/MovieSceneAnimationMixerTrack.cpp b/Engine/Plugins/Experimental/MovieSceneAnimMixer/Source/MovieSceneAnimMixer/Private/MovieSceneAnimationMixerTrack.cpp
--- a/Engine/Plugins/Experimental/MovieSceneAnimMixer/Source/MovieSceneAnimMixer/Private/MovieSceneAnimationMixerTrack.cpp
+++ b/Engine/Plugins/Experimental/MovieSceneAnimMixer/Source/MovieSceneAnimMixer/Private/MovieSceneAnimationMixerTrack.cpp
@@ -518,5 +518,47 @@ UMovieSceneAnimationMixerLayer* UMovieSceneAnimationMixerTrack::GetLayerForChild
 	return nullptr;
 }
 
+UMovieSceneAnimationMixerLayer* UMovieSceneAnimationMixerTrack::InsertLayer(int32 InsertIndex)
+{
+	// Clamp insert index to valid range
+	InsertIndex = FMath::Clamp(InsertIndex, 0, Layers.Num());
+
+	// Create the new layer
+	UMovieSceneAnimationMixerLayer* NewLayer = NewObject<UMovieSceneAnimationMixerLayer>(this, NAME_None, RF_Transactional);
+
+	// Insert at the specified position
+	Layers.Insert(NewLayer, InsertIndex);
+
+	// Update row indices for all layers from InsertIndex onwards
+	for (int32 LayerIndex = InsertIndex; LayerIndex < Layers.Num(); ++LayerIndex)
+	{
+		UMovieSceneAnimationMixerLayer* Layer = Layers[LayerIndex];
+		if (!Layer)
+		{
+			continue;
+		}
+
+		// Update sections in this layer
+		for (UMovieSceneSection* Section : Layer->GetSections())
+		{
+			if (Section)
+			{
+				Section->Modify();
+				Section->SetRowIndex(LayerIndex);
+			}
+		}
+
+		// Update child track row index if this layer has one
+		if (UMovieSceneTrack* ChildTrack = Layer->GetChildTrack())
+		{
+			if (int32* TrackRowIndex = ChildTracks.Find(ChildTrack))
+			{
+				*TrackRowIndex = LayerIndex;
+			}
+		}
+	}
+
+	return NewLayer;
+}
 
 #undef LOCTEXT_NAMESPACE
\ No newline at end of file
diff --git a/Engine/Plugins/Experimental/MovieSceneAnimMixer/Source/MovieSceneAnimMixer/Public/MovieSceneAnimationMixerTrack.h b/Engine/Plugins/Experimental/MovieSceneAnimMixer/Source/MovieSceneAnimMixer/Public/MovieSceneAnimationMixerTrack.h
--- a/Engine/Plugins/Experimental/MovieSceneAnimMixer/Source/MovieSceneAnimMixer/Public/MovieSceneAnimationMixerTrack.h
+++ b/Engine/Plugins/Experimental/MovieSceneAnimMixer/Source/MovieSceneAnimMixer/Public/MovieSceneAnimationMixerTrack.h
@@ -164,6 +164,9 @@ public:
 	/** Get the layer containing the given child track */
 	MOVIESCENEANIMMIXER_API UMovieSceneAnimationMixerLayer* GetLayerForChildTrack(UMovieSceneTrack* Track) const;
 
+	/** Insert a new empty layer at the specified index, shifting existing layers down */
+	MOVIESCENEANIMMIXER_API UMovieSceneAnimationMixerLayer* InsertLayer(int32 InsertIndex);
+
 public: 
 	UPROPERTY(EditAnywhere, Category = "Mixed Animation")
 	TInstancedStruct<FMovieSceneMixedAnimationTarget> MixedAnimationTarget;
diff --git a/Engine/Plugins/Experimental/MovieSceneAnimMixer/Source/MovieSceneAnimMixerEditor/Private/MVVM/AnimationMixerTrackModel.cpp b/Engine/Plugins/Experimental/MovieSceneAnimMixer/Source/MovieSceneAnimMixerEditor/Private/MVVM/AnimationMixerTrackModel.cpp
--- a/Engine/Plugins/Experimental/MovieSceneAnimMixer/Source/MovieSceneAnimMixerEditor/Private/MVVM/AnimationMixerTrackModel.cpp
+++ b/Engine/Plugins/Experimental/MovieSceneAnimMixer/Source/MovieSceneAnimMixerEditor/Private/MVVM/AnimationMixerTrackModel.cpp
@@ -27,6 +27,7 @@
 #include "Styling/SlateTypes.h"
 #include "Rendering/DrawElements.h"
 #include "MVVM/ViewModels/VirtualTrackArea.h"
+#include "Framework/Application/SlateApplication.h"
 
 namespace UE::Sequencer
 {
@@ -203,24 +204,7 @@ void FAnimationMixerTrackModel::ShiftRowsDownFromIndex(int32 StartIndex, const T
 	}
 
 	MixerTrack->Modify();
-
-	// Insert a new empty layer at StartIndex to shift layers from that index down
-	TArray<TObjectPtr<UMovieSceneAnimationMixerLayer>>& Layers = const_cast<TArray<TObjectPtr<UMovieSceneAnimationMixerLayer>>&>(MixerTrack->GetLayers());
-	if (StartIndex <= Layers.Num())
-	{
-		UMovieSceneAnimationMixerLayer* NewEmptyLayer = NewObject<UMovieSceneAnimationMixerLayer>(MixerTrack, NAME_None, RF_Transactional);
-		Layers.Insert(NewEmptyLayer, StartIndex);
-
-		// Update row indices for all sections and child tracks in layers from StartIndex onwards
-		for (int32 LayerIndex = StartIndex; LayerIndex < Layers.Num(); ++LayerIndex)
-		{
-			UMovieSceneAnimationMixerLayer* Layer = Layers[LayerIndex];
-			if (Layer)
-			{
-				AnimationMixerHelpers::UpdateLayerRowIndices(Layer, MixerTrack, LayerIndex, &ExcludeSections);
-			}
-		}
-	}
+	MixerTrack->InsertLayer(StartIndex);
 
 	// Rebuild layer models with updated indices
 	ForceUpdate();
@@ -324,6 +308,10 @@ TOptional<FAnimMixerDropPreview> FAnimationMixerTrackModel::ComputeDropPreview(c
 
 	const TArray<TObjectPtr<UMovieSceneAnimationMixerLayer>>& Layers = MixerTrack->GetLayers();
 
+	// Check if ctrl key is held - allows insertion between layers
+	const bool bCtrlHeld = FSlateApplication::Get().GetModifierKeys().IsControlDown();
+	const int32 LastLayerIndex = Layers.Num() - 1;
+
 	// Handle single-layer collapsed view
 	if (Layers.Num() == 1 && bIsInDrag)
 	{
@@ -338,12 +326,12 @@ TOptional<FAnimMixerDropPreview> FAnimationMixerTrackModel::ComputeDropPreview(c
 
 		if (Context.VirtualMousePos.Y < TrackTop || Context.LocalMousePos.Y <= 0)
 		{
-			// Above track - insert above layer 0
+			// Above track - always allowed (new layer at top)
 			return FAnimMixerDropPreview{ FAnimMixerDropPreview::EZone::AboveLayer, 0, TrackTop, TrackBottom };
 		}
 		else if (Context.VirtualMousePos.Y > TrackBottom)
 		{
-			// Below track - insert below layer 0
+			// Below track - always allowed (new layer at bottom)
 			return FAnimMixerDropPreview{ FAnimMixerDropPreview::EZone::BelowLayer, 0, TrackTop, TrackBottom };
 		}
 		else
@@ -375,7 +363,7 @@ TOptional<FAnimMixerDropPreview> FAnimationMixerTrackModel::ComputeDropPreview(c
 		float VirtualChildBottom = ChildVirtualGeometry.NestedBottom;
 		float VirtualChildHeight = VirtualChildBottom - VirtualChildTop;
 
-		// Check if above all rows
+		// Check if above all rows - always allowed (new layer at top)
 		if (LastRowIndex == -1 && (Context.VirtualMousePos.Y <= VirtualChildTop || Context.LocalMousePos.Y <= 0))
 		{
 			return FAnimMixerDropPreview{ FAnimMixerDropPreview::EZone::AboveLayer, 0, VirtualChildTop, VirtualChildBottom };
@@ -390,21 +378,58 @@ TOptional<FAnimMixerDropPreview> FAnimationMixerTrackModel::ComputeDropPreview(c
 
 			if (LocalMouseY < ZoneBoundary)
 			{
-				// Top 25% - insert above
-				return FAnimMixerDropPreview{ FAnimMixerDropPreview::EZone::AboveLayer, CurrentRowIndex, VirtualChildTop, VirtualChildBottom };
+				// Top 25% zone
+				// Insert above is only allowed if:
+				// - This is the first layer (row 0) - always allowed as "new layer at top"
+				// - Ctrl is held - allows insertion between layers
+				if (CurrentRowIndex == 0 || bCtrlHeld)
+				{
+					return FAnimMixerDropPreview{ FAnimMixerDropPreview::EZone::AboveLayer, CurrentRowIndex, VirtualChildTop, VirtualChildBottom };
+				}
+				else
+				{
+					// Without ctrl, drop onto this layer instead
+					if (bIsChildTrackRow)
+					{
+						// Can't drop onto child track - no preview
+						return TOptional<FAnimMixerDropPreview>();
+					}
+					return FAnimMixerDropPreview{ FAnimMixerDropPreview::EZone::OntoLayer, CurrentRowIndex, VirtualChildTop, VirtualChildBottom };
+				}
 			}
 			else if (LocalMouseY > VirtualChildHeight - ZoneBoundary)
 			{
-				// Bottom 25% - insert below
-				return FAnimMixerDropPreview{ FAnimMixerDropPreview::EZone::BelowLayer, CurrentRowIndex, VirtualChildTop, VirtualChildBottom };
+				// Bottom 25% zone
+				// Insert below is only allowed if:
+				// - This is the last layer - always allowed as "new layer at bottom"
+				// - Ctrl is held - allows insertion between layers
+				if (CurrentRowIndex == LastLayerIndex || bCtrlHeld)
+				{
+					return FAnimMixerDropPreview{ FAnimMixerDropPreview::EZone::BelowLayer, CurrentRowIndex, VirtualChildTop, VirtualChildBottom };
+				}
+				else
+				{
+					// Without ctrl, drop onto this layer instead
+					if (bIsChildTrackRow)
+					{
+						// Can't drop onto child track - no preview
+						return TOptional<FAnimMixerDropPreview>();
+					}
+					return FAnimMixerDropPreview{ FAnimMixerDropPreview::EZone::OntoLayer, CurrentRowIndex, VirtualChildTop, VirtualChildBottom };
+				}
 			}
 			else
 			{
 				// Middle 50%
 				if (bIsChildTrackRow)
 				{
-					// Can't drop onto child track - default to below
-					return FAnimMixerDropPreview{ FAnimMixerDropPreview::EZone::BelowLayer, CurrentRowIndex, VirtualChildTop, VirtualChildBottom };
+					// Can't drop onto child track
+					// If ctrl is held, allow insert below, otherwise no preview
+					if (bCtrlHeld || CurrentRowIndex == LastLayerIndex)
+					{
+						return FAnimMixerDropPreview{ FAnimMixerDropPreview::EZone::BelowLayer, CurrentRowIndex, VirtualChildTop, VirtualChildBottom };
+					}
+					return TOptional<FAnimMixerDropPreview>();
 				}
 				else
 				{
@@ -417,7 +442,7 @@ TOptional<FAnimMixerDropPreview> FAnimationMixerTrackModel::ComputeDropPreview(c
 		LastRowIndex = CurrentRowIndex;
 	}
 
-	// Below all rows - insert at end
+	// Below all rows - always allowed (new layer at bottom)
 	if (Context.VirtualMousePos.Y >= 0 && LastRowIndex >= 0)
 	{
 		// Get geometry of last layer
diff --git a/Engine/Plugins/Experimental/MovieSceneAnimMixer/Source/MovieSceneAnimMixerEditor/Private/MovieSceneAnimationMixerTrackEditor.cpp b/Engine/Plugins/Experimental/MovieSceneAnimMixer/Source/MovieSceneAnimMixerEditor/Private/MovieSceneAnimationMixerTrackEditor.cpp
--- a/Engine/Plugins/Experimental/MovieSceneAnimMixer/Source/MovieSceneAnimMixerEditor/Private/MovieSceneAnimationMixerTrackEditor.cpp
+++ b/Engine/Plugins/Experimental/MovieSceneAnimMixer/Source/MovieSceneAnimMixerEditor/Private/MovieSceneAnimationMixerTrackEditor.cpp
@@ -772,6 +772,61 @@ TSharedRef<SWidget> FAnimationMixerTrackEditor::BuildAddSectionSubMenu(TWeakView
 		}
 		MenuBuilder.EndSection();
 	}
+	else
+	{
+		// Insert layer options - only show on layer rows
+		const int32 CurrentLayerIndex = LayerExtension->GetRowIndex();
+
+		MenuBuilder.BeginSection(NAME_None, LOCTEXT("LayerManagement_Label", "Layer Management"));
+		{
+			MenuBuilder.AddMenuEntry(
+				LOCTEXT("InsertLayerAbove", "Insert Layer Above"),
+				LOCTEXT("InsertLayerAbove_Tooltip", "Insert a new empty layer above this one"),
+				FSlateIcon(),
+				FUIAction(FExecuteAction::CreateLambda([WeakViewModel, WeakEditor, CurrentLayerIndex]()
+				{
+					TViewModelPtr<ITrackRowExtension>        LayerExt  = WeakViewModel.ImplicitPin();
+					TViewModelPtr<FSequencerEditorViewModel> Editor    = WeakEditor.Pin();
+					TSharedPtr<ISequencer>                   Sequencer = Editor ? Editor->GetSequencer() : nullptr;
+
+					UMovieSceneAnimationMixerTrack* MixerTrack = LayerExt ? Cast<UMovieSceneAnimationMixerTrack>(LayerExt->GetParentTrack()) : nullptr;
+					if (Sequencer && MixerTrack)
+					{
+						FScopedTransaction Transaction(LOCTEXT("InsertLayerAboveTransaction", "Insert Layer Above"));
+						MixerTrack->Modify();
+
+						MixerTrack->InsertLayer(CurrentLayerIndex);
+
+						Sequencer->NotifyMovieSceneDataChanged(EMovieSceneDataChangeType::MovieSceneStructureItemsChanged);
+					}
+				}))
+			);
+
+			MenuBuilder.AddMenuEntry(
+				LOCTEXT("InsertLayerBelow", "Insert Layer Below"),
+				LOCTEXT("InsertLayerBelow_Tooltip", "Insert a new empty layer below this one"),
+				FSlateIcon(),
+				FUIAction(FExecuteAction::CreateLambda([WeakViewModel, WeakEditor, CurrentLayerIndex]()
+				{
+					TViewModelPtr<ITrackRowExtension>        LayerExt  = WeakViewModel.ImplicitPin();
+					TViewModelPtr<FSequencerEditorViewModel> Editor    = WeakEditor.Pin();
+					TSharedPtr<ISequencer>                   Sequencer = Editor ? Editor->GetSequencer() : nullptr;
+
+					UMovieSceneAnimationMixerTrack* MixerTrack = LayerExt ? Cast<UMovieSceneAnimationMixerTrack>(LayerExt->GetParentTrack()) : nullptr;
+					if (Sequencer && MixerTrack)
+					{
+						FScopedTransaction Transaction(LOCTEXT("InsertLayerBelowTransaction", "Insert Layer Below"));
+						MixerTrack->Modify();
+
+						MixerTrack->InsertLayer(CurrentLayerIndex + 1);
+
+						Sequencer->NotifyMovieSceneDataChanged(EMovieSceneDataChangeType::MovieSceneStructureItemsChanged);
+					}
+				}))
+			);
+		}
+		MenuBuilder.EndSection();
+	}
 
 	if (TViewModelPtr<ISectionOwnerExtension> SectionOwner = WeakViewModel.ImplicitPin())
 	{ 
