Sequencer: Anim Mixer UX Minor fixes
* Change preview for section dragging to be the size of the section and re-use the current drop target UX style for it
* Remove 25% drop targets completely if ctrl is not held including at top and bottom
* Fix issue where target drop-down box is also on layers
* Ignore hover highlighting when dragging sections over rows outside of the current track

#Synced-CL 49155801

diff --git a/Engine/Plugins/Experimental/MovieSceneAnimMixer/Source/MovieSceneAnimMixerEditor/Private/MVVM/AnimationMixerTrackModel.cpp b/Engine/Plugins/Experimental/MovieSceneAnimMixer/Source/MovieSceneAnimMixerEditor/Private/MVVM/AnimationMixerTrackModel.cpp
--- a/Engine/Plugins/Experimental/MovieSceneAnimMixer/Source/MovieSceneAnimMixerEditor/Private/MVVM/AnimationMixerTrackModel.cpp
+++ b/Engine/Plugins/Experimental/MovieSceneAnimMixer/Source/MovieSceneAnimMixerEditor/Private/MVVM/AnimationMixerTrackModel.cpp
@@ -25,6 +25,8 @@
 #include "MovieSceneSection.h"
 #include "Styling/CoreStyle.h"
 #include "Styling/SlateTypes.h"
+#include "Styling/AppStyle.h"
+#include "Styling/StyleColors.h"
 #include "Rendering/DrawElements.h"
 #include "MVVM/ViewModels/VirtualTrackArea.h"
 #include "Framework/Application/SlateApplication.h"
@@ -254,29 +256,80 @@ int32 FAnimationMixerTrackModel::OnPaintSectionDragPreview(
 	}
 
 	const FAnimMixerDropPreview& Preview = CachedPreview.GetValue();
-	const FTableRowStyle* RowStyle = &FCoreStyle::Get().GetWidgetStyle<FTableRowStyle>("TableView.Row");
 
-	const FSlateBrush* Brush = nullptr;
-	switch (Preview.Zone)
+	// Convert virtual coordinates to physical coordinates
+	FVector2D PhysicalTop = VirtualTrackArea.VirtualToPhysical(FVector2D(0.0f, Preview.VirtualTop));
+	FVector2D PhysicalBottom = VirtualTrackArea.VirtualToPhysical(FVector2D(0.0f, Preview.VirtualBottom));
+	float PhysicalHeight = PhysicalBottom.Y - PhysicalTop.Y;
+
+	if (Preview.Zone == FAnimMixerDropPreview::EZone::OntoLayer)
 	{
-		case FAnimMixerDropPreview::EZone::AboveLayer:
-			Brush = &RowStyle->DropIndicator_Above;
-			break;
-		case FAnimMixerDropPreview::EZone::BelowLayer:
-			Brush = &RowStyle->DropIndicator_Below;
-			break;
-		case FAnimMixerDropPreview::EZone::OntoLayer:
-			Brush = &RowStyle->DropIndicator_Onto;
-			break;
-	}
+		// Draw dashed outline around the section bounds (matching STrackAreaView drop target style)
+		FSlateColor DashColor = FStyleColors::AccentBlue;
+
+		const FSlateBrush* HorizontalBrush = FAppStyle::GetBrush("WideDash.Horizontal");
+		const FSlateBrush* VerticalBrush = FAppStyle::GetBrush("WideDash.Vertical");
+
+		// Calculate horizontal bounds from section frame range
+		float DropMinX = 0.f;
+		float DropMaxX = TrackAreaGeometry.GetLocalSize().X;
+
+		if (Preview.SectionFrameRange.HasLowerBound() && Preview.SectionFrameRange.HasUpperBound())
+		{
+			DropMinX = VirtualTrackArea.FrameToPixel(Preview.SectionFrameRange.GetLowerBoundValue());
+			DropMaxX = VirtualTrackArea.FrameToPixel(Preview.SectionFrameRange.GetUpperBoundValue());
+		}
+
+		float TrackPosition = PhysicalTop.Y;
+		float TrackHeight = PhysicalHeight;
+
+		// Top edge
+		FSlateDrawElement::MakeBox(
+			OutDrawElements,
+			LayerId,
+			TrackAreaGeometry.ToPaintGeometry(FVector2f(DropMaxX - DropMinX, HorizontalBrush->ImageSize.Y), FSlateLayoutTransform(FVector2f(DropMinX, TrackPosition))),
+			HorizontalBrush,
+			ESlateDrawEffect::None,
+			DashColor.GetSpecifiedColor());
+
+		// Bottom edge
+		FSlateDrawElement::MakeBox(
+			OutDrawElements,
+			LayerId,
+			TrackAreaGeometry.ToPaintGeometry(FVector2f(DropMaxX - DropMinX, HorizontalBrush->ImageSize.Y), FSlateLayoutTransform(FVector2f(DropMinX, TrackPosition + (TrackHeight - HorizontalBrush->ImageSize.Y)))),
+			HorizontalBrush,
+			ESlateDrawEffect::None,
+			DashColor.GetSpecifiedColor());
 
-	if (Brush)
+		// Left edge
+		FSlateDrawElement::MakeBox(
+			OutDrawElements,
+			LayerId,
+			TrackAreaGeometry.ToPaintGeometry(FVector2f(VerticalBrush->ImageSize.X, TrackHeight), FSlateLayoutTransform(FVector2f(DropMinX, TrackPosition))),
+			VerticalBrush,
+			ESlateDrawEffect::None,
+			DashColor.GetSpecifiedColor());
+
+		// Right edge
+		FSlateDrawElement::MakeBox(
+			OutDrawElements,
+			LayerId,
+			TrackAreaGeometry.ToPaintGeometry(FVector2f(VerticalBrush->ImageSize.X, TrackHeight), FSlateLayoutTransform(FVector2f(DropMaxX - VerticalBrush->ImageSize.X, TrackPosition))),
+			VerticalBrush,
+			ESlateDrawEffect::None,
+			DashColor.GetSpecifiedColor());
+
+		LayerId++;
+	}
+	else
 	{
-		// Convert virtual coordinates to physical coordinates
-		FVector2D PhysicalTop = VirtualTrackArea.VirtualToPhysical(FVector2D(0.0f, Preview.VirtualTop));
-		FVector2D PhysicalBottom = VirtualTrackArea.VirtualToPhysical(FVector2D(0.0f, Preview.VirtualBottom));
+		// AboveLayer/BelowLayer - use row-style drop indicators
+		const FTableRowStyle* RowStyle = &FCoreStyle::Get().GetWidgetStyle<FTableRowStyle>("TableView.Row");
+
+		const FSlateBrush* Brush = (Preview.Zone == FAnimMixerDropPreview::EZone::AboveLayer)
+			? &RowStyle->DropIndicator_Above
+			: &RowStyle->DropIndicator_Below;
 
-		float PhysicalHeight = PhysicalBottom.Y - PhysicalTop.Y;
 		FVector2D LocalSize(TrackAreaGeometry.GetLocalSize().X, PhysicalHeight);
 		FVector2D LocalPosition(0.0f, PhysicalTop.Y);
 
@@ -306,6 +359,20 @@ TOptional<FAnimMixerDropPreview> FAnimationMixerTrackModel::ComputeDropPreview(c
 		return TOptional<FAnimMixerDropPreview>();
 	}
 
+	// Compute combined frame range of all dragged sections (for outline drawing)
+	TRange<FFrameNumber> DraggedSectionRange = TRange<FFrameNumber>::Empty();
+	for (UMovieSceneSection* Section : *Context.DraggedSections)
+	{
+		if (Section)
+		{
+			TRange<FFrameNumber> SectionRange = Section->GetRange();
+			if (!SectionRange.IsEmpty())
+			{
+				DraggedSectionRange = TRange<FFrameNumber>::Hull(DraggedSectionRange, SectionRange);
+			}
+		}
+	}
+
 	const TArray<TObjectPtr<UMovieSceneAnimationMixerLayer>>& Layers = MixerTrack->GetLayers();
 
 	// Check if ctrl key is held - allows insertion between layers
@@ -337,7 +404,7 @@ TOptional<FAnimMixerDropPreview> FAnimationMixerTrackModel::ComputeDropPreview(c
 		else
 		{
 			// Within track - drop onto layer 0
-			return FAnimMixerDropPreview{ FAnimMixerDropPreview::EZone::OntoLayer, 0, TrackTop, TrackBottom };
+			return FAnimMixerDropPreview{ FAnimMixerDropPreview::EZone::OntoLayer, 0, TrackTop, TrackBottom, DraggedSectionRange };
 		}
 	}
 
@@ -378,11 +445,8 @@ TOptional<FAnimMixerDropPreview> FAnimationMixerTrackModel::ComputeDropPreview(c
 
 			if (LocalMouseY < ZoneBoundary)
 			{
-				// Top 25% zone
-				// Insert above is only allowed if:
-				// - This is the first layer (row 0) - always allowed as "new layer at top"
-				// - Ctrl is held - allows insertion between layers
-				if (CurrentRowIndex == 0 || bCtrlHeld)
+				// Top 25% zone- only use insertion drop zone in ctrl insertion case
+				if (bCtrlHeld)
 				{
 					return FAnimMixerDropPreview{ FAnimMixerDropPreview::EZone::AboveLayer, CurrentRowIndex, VirtualChildTop, VirtualChildBottom };
 				}
@@ -394,16 +458,13 @@ TOptional<FAnimMixerDropPreview> FAnimationMixerTrackModel::ComputeDropPreview(c
 						// Can't drop onto child track - no preview
 						return TOptional<FAnimMixerDropPreview>();
 					}
-					return FAnimMixerDropPreview{ FAnimMixerDropPreview::EZone::OntoLayer, CurrentRowIndex, VirtualChildTop, VirtualChildBottom };
+					return FAnimMixerDropPreview{ FAnimMixerDropPreview::EZone::OntoLayer, CurrentRowIndex, VirtualChildTop, VirtualChildBottom, DraggedSectionRange };
 				}
 			}
 			else if (LocalMouseY > VirtualChildHeight - ZoneBoundary)
 			{
-				// Bottom 25% zone
-				// Insert below is only allowed if:
-				// - This is the last layer - always allowed as "new layer at bottom"
-				// - Ctrl is held - allows insertion between layers
-				if (CurrentRowIndex == LastLayerIndex || bCtrlHeld)
+				// Bottom 25% zone- only use insertion drop zones in ctrl insertion case
+				if (bCtrlHeld)
 				{
 					return FAnimMixerDropPreview{ FAnimMixerDropPreview::EZone::BelowLayer, CurrentRowIndex, VirtualChildTop, VirtualChildBottom };
 				}
@@ -415,7 +476,7 @@ TOptional<FAnimMixerDropPreview> FAnimationMixerTrackModel::ComputeDropPreview(c
 						// Can't drop onto child track - no preview
 						return TOptional<FAnimMixerDropPreview>();
 					}
-					return FAnimMixerDropPreview{ FAnimMixerDropPreview::EZone::OntoLayer, CurrentRowIndex, VirtualChildTop, VirtualChildBottom };
+					return FAnimMixerDropPreview{ FAnimMixerDropPreview::EZone::OntoLayer, CurrentRowIndex, VirtualChildTop, VirtualChildBottom, DraggedSectionRange };
 				}
 			}
 			else
@@ -434,7 +495,7 @@ TOptional<FAnimMixerDropPreview> FAnimationMixerTrackModel::ComputeDropPreview(c
 				else
 				{
 					// Drop onto section layer
-					return FAnimMixerDropPreview{ FAnimMixerDropPreview::EZone::OntoLayer, CurrentRowIndex, VirtualChildTop, VirtualChildBottom };
+					return FAnimMixerDropPreview{ FAnimMixerDropPreview::EZone::OntoLayer, CurrentRowIndex, VirtualChildTop, VirtualChildBottom, DraggedSectionRange };
 				}
 			}
 		}
diff --git a/Engine/Plugins/Experimental/MovieSceneAnimMixer/Source/MovieSceneAnimMixerEditor/Private/MovieSceneAnimationMixerTrackEditor.cpp b/Engine/Plugins/Experimental/MovieSceneAnimMixer/Source/MovieSceneAnimMixerEditor/Private/MovieSceneAnimationMixerTrackEditor.cpp
--- a/Engine/Plugins/Experimental/MovieSceneAnimMixer/Source/MovieSceneAnimMixerEditor/Private/MovieSceneAnimationMixerTrackEditor.cpp
+++ b/Engine/Plugins/Experimental/MovieSceneAnimMixer/Source/MovieSceneAnimMixerEditor/Private/MovieSceneAnimationMixerTrackEditor.cpp
@@ -29,6 +29,7 @@
 #include "SequencerUtilities.h"
 #include "Widgets/Input/SComboButton.h"
 #include "Widgets/Images/SImage.h"
+#include "MVVM/AnimationMixerTrackModel.h"
 
 #define LOCTEXT_NAMESPACE "MovieSceneAnimationMixerTrackEditor"
 
@@ -547,7 +548,7 @@ TSharedPtr<SWidget> FAnimationMixerTrackEditor::BuildOutlinerColumnWidget(const
 	if (ColumnName == FCommonOutlinerNames::Edit)
 	{
 		// Only show target combo on the track itself, not on layers
-		if (TViewModelPtr<ITrackExtension> TrackModel = Params.ViewModel.ImplicitCast())
+		if (TViewModelPtr<FAnimationMixerTrackModel> TrackModel = Params.ViewModel.ImplicitCast())
 		{
 			if (UMovieSceneAnimationMixerTrack* MixerTrack = Cast<UMovieSceneAnimationMixerTrack>(TrackModel->GetTrack()))
 			{
diff --git a/Engine/Plugins/Experimental/MovieSceneAnimMixer/Source/MovieSceneAnimMixerEditor/Public/MVVM/AnimationMixerTrackModel.h b/Engine/Plugins/Experimental/MovieSceneAnimMixer/Source/MovieSceneAnimMixerEditor/Public/MVVM/AnimationMixerTrackModel.h
--- a/Engine/Plugins/Experimental/MovieSceneAnimMixer/Source/MovieSceneAnimMixerEditor/Public/MVVM/AnimationMixerTrackModel.h
+++ b/Engine/Plugins/Experimental/MovieSceneAnimMixer/Source/MovieSceneAnimMixerEditor/Public/MVVM/AnimationMixerTrackModel.h
@@ -3,6 +3,8 @@
 #pragma once
 
 #include "CoreTypes.h"
+#include "Misc/FrameNumber.h"
+#include "Math/Range.h"
 #include "MVVM/ViewModels/TrackModel.h"
 #include "MVVM/ViewModelPtr.h"
 
@@ -25,8 +27,9 @@ struct FAnimMixerDropPreview
 
 	EZone Zone = EZone::OntoLayer;
 	int32 TargetRowIndex = INDEX_NONE;
-	float VirtualTop = 0.0f;       
-	float VirtualBottom = 0.0f;    
+	float VirtualTop = 0.0f;
+	float VirtualBottom = 0.0f;
+	TRange<FFrameNumber> SectionFrameRange; 
 
 	bool operator==(const FAnimMixerDropPreview& Other) const
 	{
diff --git a/Engine/Source/Editor/Sequencer/Private/Tools/EditToolDragOperations.cpp b/Engine/Source/Editor/Sequencer/Private/Tools/EditToolDragOperations.cpp
--- a/Engine/Source/Editor/Sequencer/Private/Tools/EditToolDragOperations.cpp
+++ b/Engine/Source/Editor/Sequencer/Private/Tools/EditToolDragOperations.cpp
@@ -732,7 +732,7 @@ void FMoveKeysAndSections::AddModel(TSharedPtr<UE::Sequencer::FViewModel> Model)
 		if (UMovieSceneSection* Section = SectionModel->GetSection())
 		{
 			Sections.Add(Section);
-	
+
 			if (UMovieScene* MovieScene = Section->GetTypedOuter<UMovieScene>())
 			{
 				// If the section is in a group, we also want to add the sections it is grouped with
@@ -753,9 +753,23 @@ void FMoveKeysAndSections::AddModel(TSharedPtr<UE::Sequencer::FViewModel> Model)
 	if (TSharedPtr<IDraggableTrackAreaExtension> DraggableItem = Model->CastThisShared<IDraggableTrackAreaExtension>())
 	{
 		DraggedItems.Add(DraggableItem);
+
+		// Cache the drag owner track model if not already set (used for hover suppression during drag)
+		if (!CachedDragOwnerTrackModel.IsValid())
+		{
+			if (TViewModelPtr<ISectionDragOwnerExtension> TrackExt = Model->FindAncestorOfType<ISectionDragOwnerExtension>(true))
+			{
+				CachedDragOwnerTrackModel = TrackExt.AsModel();
+			}
+		}
 	}
 }
 
+TSharedPtr<UE::Sequencer::FViewModel> FMoveKeysAndSections::GetDragOwnerTrackModel() const
+{
+	return CachedDragOwnerTrackModel.Pin();
+}
+
 void FMoveKeysAndSections::OnBeginDrag(const FPointerEvent& MouseEvent, FVector2D LocalMousePos, const UE::Sequencer::FVirtualTrackArea& VirtualTrackArea)
 {
 	using namespace UE::Sequencer;
diff --git a/Engine/Source/Editor/Sequencer/Private/Tools/EditToolDragOperations.h b/Engine/Source/Editor/Sequencer/Private/Tools/EditToolDragOperations.h
--- a/Engine/Source/Editor/Sequencer/Private/Tools/EditToolDragOperations.h
+++ b/Engine/Source/Editor/Sequencer/Private/Tools/EditToolDragOperations.h
@@ -178,6 +178,9 @@ public:
 	virtual void AddSnapTime(FFrameNumber SnapTime) override;
 	virtual void AddModel(TSharedPtr<UE::Sequencer::FViewModel> Model) override;
 
+	/* ISequencerEditToolDragOperation */
+	virtual TSharedPtr<UE::Sequencer::FViewModel> GetDragOwnerTrackModel() const override;
+
 protected:
 	/** Calculate the possible horizontal movement we can, constrained by sections running into things. */
 	TOptional<FFrameNumber> GetMovementDeltaX(const FPointerEvent& MouseEvent, FFrameTime MouseTime);
@@ -229,6 +232,9 @@ protected:
 	/** If the user is moving them via clicking on the Section then we'll allow vertical re-arranging, otherwise not. */
 	bool bAllowVerticalMovement;
 
+	/** Cached drag owner track model (computed in AddModel) */
+	TWeakPtr<UE::Sequencer::FViewModel> CachedDragOwnerTrackModel;
+
 protected:
 	void GetSections(TArray<UMovieSceneSection*>& OutSections) override { OutSections = Sections.Array(); }
 };
diff --git a/Engine/Source/Editor/Sequencer/Private/Tools/SequencerEditTool_Movement.cpp b/Engine/Source/Editor/Sequencer/Private/Tools/SequencerEditTool_Movement.cpp
--- a/Engine/Source/Editor/Sequencer/Private/Tools/SequencerEditTool_Movement.cpp
+++ b/Engine/Source/Editor/Sequencer/Private/Tools/SequencerEditTool_Movement.cpp
@@ -618,3 +618,8 @@ TSharedPtr<UE::Sequencer::ITrackAreaHotspot> FSequencerEditTool_Movement::GetDra
 {
 	return DelayedDrag.IsSet() ? DelayedDrag->Hotspot : nullptr;
 }
+
+UE::Sequencer::ISequencerEditToolDragOperation* FSequencerEditTool_Movement::GetCurrentDragOperation() const
+{
+	return DragOperation.Get();
+}
diff --git a/Engine/Source/Editor/Sequencer/Private/Tools/SequencerEditTool_Movement.h b/Engine/Source/Editor/Sequencer/Private/Tools/SequencerEditTool_Movement.h
--- a/Engine/Source/Editor/Sequencer/Private/Tools/SequencerEditTool_Movement.h
+++ b/Engine/Source/Editor/Sequencer/Private/Tools/SequencerEditTool_Movement.h
@@ -42,6 +42,7 @@ public:
 	virtual FName GetIdentifier() const override;
 	virtual bool CanDeactivate() const override;
 	virtual TSharedPtr<UE::Sequencer::ITrackAreaHotspot> GetDragHotspot() const override;
+	virtual UE::Sequencer::ISequencerEditToolDragOperation* GetCurrentDragOperation() const override;
 
 protected:
 
diff --git a/Engine/Source/Editor/SequencerCore/Private/MVVM/Views/STrackAreaView.cpp b/Engine/Source/Editor/SequencerCore/Private/MVVM/Views/STrackAreaView.cpp
--- a/Engine/Source/Editor/SequencerCore/Private/MVVM/Views/STrackAreaView.cpp
+++ b/Engine/Source/Editor/SequencerCore/Private/MVVM/Views/STrackAreaView.cpp
@@ -426,7 +426,28 @@ void STrackAreaView::UpdateHoverStates( const FGeometry& MyGeometry, const FPoin
 
 	// Set the node that we are hovering
 	TViewModelPtr<IOutlinerExtension> NewHoveredItem = PinnedOutliner->HitTestNode(MyGeometry.AbsoluteToLocal(MouseEvent.GetScreenSpacePosition()).Y);
+
 	TSharedPtr<FOutlinerViewModel> OutlinerModel = PinnedOutliner->GetOutlinerModel();
+
+	bool bHoveredItemChanging = OutlinerModel && (NewHoveredItem != OutlinerModel->GetHoveredItem());
+	// Suppress hover for items outside the drag owner's track during section drag
+	if (bHoveredItemChanging && NewHoveredItem && ViewModel)
+	{
+		if (ISequencerEditTool* EditTool = ViewModel->GetEditTool())
+		{
+			if (ISequencerEditToolDragOperation* DragOp = EditTool->GetCurrentDragOperation())
+			{
+				if (TSharedPtr<FViewModel> DragOwner = DragOp->GetDragOwnerTrackModel())
+				{
+					if (!DragOwner->IsAncestorOf(NewHoveredItem.AsModel(), /*bIncludeThis*/ true))
+					{
+						NewHoveredItem = nullptr;
+					}
+				}
+			}
+		}
+	}
+
 	if (OutlinerModel)
 	{
 		OutlinerModel->SetHoveredItem(NewHoveredItem);
diff --git a/Engine/Source/Editor/SequencerCore/Public/ISequencerEditTool.h b/Engine/Source/Editor/SequencerCore/Public/ISequencerEditTool.h
--- a/Engine/Source/Editor/SequencerCore/Public/ISequencerEditTool.h
+++ b/Engine/Source/Editor/SequencerCore/Public/ISequencerEditTool.h
@@ -16,6 +16,7 @@ namespace Sequencer
 
 struct ITrackAreaHotspot;
 class FVirtualTrackArea;
+class FViewModel;
 
 /**
  * Interface for drag and drop operations that are handled by edit tools in Sequencer.
@@ -56,6 +57,9 @@ public:
 	/** Override to implement drag-specific paint logic */
 	virtual int32 OnPaint(const FGeometry& AllottedGeometry, const FSlateRect& MyCullingRect, FSlateWindowElementList& OutDrawElements, int32 LayerId) const = 0;
 
+	/** Get the track model that owns this drag operation (for hover suppression). Returns nullptr if not applicable. */
+	virtual TSharedPtr<FViewModel> GetDragOwnerTrackModel() const { return nullptr; }
+
 public:
 
 	/** Virtual destructor. */
@@ -81,7 +85,8 @@ public:
 	virtual FName GetIdentifier() const = 0;
 	virtual bool CanDeactivate() const = 0;
 	virtual TSharedPtr<ITrackAreaHotspot> GetDragHotspot() const = 0;
-	
+	virtual ISequencerEditToolDragOperation* GetCurrentDragOperation() const { return nullptr; }
+
 public:
 
 	/** Virtual destructor. */
diff --git a/Engine/Source/Editor/SequencerCore/Public/MVVM/ViewModels/ViewModel.h b/Engine/Source/Editor/SequencerCore/Public/MVVM/ViewModels/ViewModel.h
--- a/Engine/Source/Editor/SequencerCore/Public/MVVM/ViewModels/ViewModel.h
+++ b/Engine/Source/Editor/SequencerCore/Public/MVVM/ViewModels/ViewModel.h
@@ -242,6 +242,25 @@ public:
 		return TParentModelIterator<T>(SharedThis(const_cast<FViewModel*>(this)), bIncludeThis);
 	}
 
+	/**
+	 * Returns true if this view model is an ancestor of the given view model.
+	 * @param InDescendant The potential descendant to check
+	 * @param bIncludeThis If true, returns true when InDescendant.Get() == this
+	 */
+	bool IsAncestorOf(const TSharedPtr<FViewModel>& InDescendant, bool bIncludeThis = false) const
+	{
+		TSharedPtr<FViewModel> Current = bIncludeThis ? InDescendant : (InDescendant ? InDescendant->WeakParent.Pin() : nullptr);
+		while (Current)
+		{
+			if (Current.Get() == this)
+			{
+				return true;
+			}
+			Current = Current->WeakParent.Pin();
+		}
+		return false;
+	}
+
 protected:
 
 	/** Chain a new children list to the existing children lists */
