Sequencer: Fix Localization warning from empty string.

#Synced-CL 49388879

diff --git a/Engine/Plugins/Experimental/MovieSceneAnimMixer/Source/MovieSceneAnimMixerEditor/Private/AnimMixerTargetMenuProviders.cpp b/Engine/Plugins/Experimental/MovieSceneAnimMixer/Source/MovieSceneAnimMixerEditor/Private/AnimMixerTargetMenuProviders.cpp
--- a/Engine/Plugins/Experimental/MovieSceneAnimMixer/Source/MovieSceneAnimMixerEditor/Private/AnimMixerTargetMenuProviders.cpp
+++ b/Engine/Plugins/Experimental/MovieSceneAnimMixer/Source/MovieSceneAnimMixerEditor/Private/AnimMixerTargetMenuProviders.cpp
@@ -230,7 +230,7 @@ void FAnimNextInjectionTargetMenuProvider::PopulateAnimNextTargetSubmenu(
 		SNew(SVariablePickerCombo)
 		.PickerArgs(PickerArgs)
 		.VariableName(LOCTEXT("SelectInjectionSite", "Select Injection Site...")),
-		LOCTEXT("InjectionSitePickerLabel", ""),
+		FText(),
 		/*bNoIndent=*/ true
 	);
 }
