Sequencer: Fix another localization warning.
#Synced-CL 49388879

diff --git a/Engine/Plugins/Experimental/MovieSceneAnimMixer/Source/MovieSceneAnimMixer/Private/MovieSceneAnimationMixerLayer.cpp b/Engine/Plugins/Experimental/MovieSceneAnimMixer/Source/MovieSceneAnimMixer/Private/MovieSceneAnimationMixerLayer.cpp
--- a/Engine/Plugins/Experimental/MovieSceneAnimMixer/Source/MovieSceneAnimMixer/Private/MovieSceneAnimationMixerLayer.cpp
+++ b/Engine/Plugins/Experimental/MovieSceneAnimMixer/Source/MovieSceneAnimMixer/Private/MovieSceneAnimationMixerLayer.cpp
@@ -92,7 +92,7 @@ FText UMovieSceneAnimationMixerLayer::GetDisplayName() const
 	// Generate default name based on content
 	if (ChildTrack)
 	{
-		return FText::Format(NSLOCTEXT("MovieSceneAnimationMixerLayer", "LayerNameFormat", "{0}: {1}"), LayerNumberText, FText::FromString(ChildTrack->GetTrackName().ToString()));		
+		return FText::Format(NSLOCTEXT("MovieSceneAnimationMixerLayer", "ChildTrackLayerNameFormat", "{0}: {1}"), LayerNumberText, FText::FromString(ChildTrack->GetTrackName().ToString()));		
 	}
 	else
 	{
