Sequencer: Add 'mirroring decoration' that allows arbitrary tracks/sections in the anim mixer to support mirroring their pose. For now, just add it to control rig tracks (as skeletal animation sections already support mirroring explicitly in their properties due to backwards compatibility).

#Synced-CL 49622961

diff --git a/Engine/Plugins/Experimental/MovieSceneAnimMixer/Source/MovieSceneAnimMixer/Private/MovieSceneMirroringDecoration.cpp b/Engine/Plugins/Experimental/MovieSceneAnimMixer/Source/MovieSceneAnimMixer/Private/MovieSceneMirroringDecoration.cpp
new file mode 100644
--- /dev/null
+++ b/Engine/Plugins/Experimental/MovieSceneAnimMixer/Source/MovieSceneAnimMixer/Private/MovieSceneMirroringDecoration.cpp
@@ -0,0 +1,21 @@
+// Copyright Epic Games, Inc. All Rights Reserved.
+
+#include "MovieSceneMirroringDecoration.h"
+#include "EntitySystem/MovieSceneEntitySystemLinker.h"
+#include "EntitySystem/MovieSceneEntityBuilder.h"
+#include "AnimMixerComponentTypes.h"
+
+void UMovieSceneMirroringDecoration::ExtendEntityImpl(UMovieSceneEntitySystemLinker* EntityLinker, const FEntityImportParams& Params, FImportedEntity* OutImportedEntity)
+{
+	using namespace UE::MovieScene;
+
+	if (MirrorDataTable)
+	{
+		FAnimMixerComponentTypes* AnimMixerComponents = FAnimMixerComponentTypes::Get();
+		OutImportedEntity->AddBuilder(
+			FEntityBuilder()
+			.Add(AnimMixerComponents->MirrorTable, MirrorDataTable)
+			.AddMutualComponents()
+		);
+	}
+}
diff --git a/Engine/Plugins/Experimental/MovieSceneAnimMixer/Source/MovieSceneAnimMixer/Public/MovieSceneMirroringDecoration.h b/Engine/Plugins/Experimental/MovieSceneAnimMixer/Source/MovieSceneAnimMixer/Public/MovieSceneMirroringDecoration.h
new file mode 100644
--- /dev/null
+++ b/Engine/Plugins/Experimental/MovieSceneAnimMixer/Source/MovieSceneAnimMixer/Public/MovieSceneMirroringDecoration.h
@@ -0,0 +1,35 @@
+// Copyright Epic Games, Inc. All Rights Reserved.
+
+#pragma once
+
+#include "CoreMinimal.h"
+#include "MovieSceneSignedObject.h"
+#include "Decorations/IMovieSceneSectionDecoration.h"
+#include "Decorations/IMovieSceneTrackDecoration.h"
+#include "EntitySystem/IMovieSceneEntityDecorator.h"
+#include "MovieSceneMirroringDecoration.generated.h"
+
+class UMirrorDataTable;
+
+/**
+ * Decoration that applies mirroring to animation poses produced by sections or tracks in the Anim Mixer.
+ *
+ */
+UCLASS(MinimalAPI, collapseCategories, meta=(DisplayName="Mirroring"))
+class UMovieSceneMirroringDecoration
+	: public UMovieSceneSignedObject
+	, public IMovieSceneSectionDecoration
+	, public IMovieSceneTrackDecoration
+	, public IMovieSceneEntityDecorator
+{
+	GENERATED_BODY()
+
+public:
+
+	/** The mirror data table used to mirror the animation pose */
+	UPROPERTY(EditAnywhere, BlueprintReadWrite, Category = "Mirroring")
+	TObjectPtr<UMirrorDataTable> MirrorDataTable;
+
+	// IMovieSceneEntityDecorator interface
+	virtual void ExtendEntityImpl(UMovieSceneEntitySystemLinker* EntityLinker, const FEntityImportParams& Params, FImportedEntity* OutImportedEntity) override;
+};
diff --git a/Engine/Plugins/Experimental/MovieSceneMixedControlRig/Source/MovieSceneMixedControlRig/Private/MovieSceneMixedControlRigModule.cpp b/Engine/Plugins/Experimental/MovieSceneMixedControlRig/Source/MovieSceneMixedControlRig/Private/MovieSceneMixedControlRigModule.cpp
--- a/Engine/Plugins/Experimental/MovieSceneMixedControlRig/Source/MovieSceneMixedControlRig/Private/MovieSceneMixedControlRigModule.cpp
+++ b/Engine/Plugins/Experimental/MovieSceneMixedControlRig/Source/MovieSceneMixedControlRig/Private/MovieSceneMixedControlRigModule.cpp
@@ -3,6 +3,9 @@
 #include "MovieSceneMixedControlRigModule.h"
 #include "Modules/ModuleInterface.h"
 #include "Modules/ModuleManager.h"
+#include "IMovieSceneModule.h"
+#include "MovieSceneMirroringDecoration.h"
+#include "Sequencer/MovieSceneControlRigParameterTrack.h"
 
 DEFINE_LOG_CATEGORY(LogMovieSceneMixedControlRig);
 
@@ -13,11 +16,19 @@ namespace UE::MovieScene
 
 void FMovieSceneMixedControlRigModule::StartupModule()
 {
+	IMovieSceneModule::Get().RegisterCompatibleDecoration(
+		UMovieSceneControlRigParameterTrack::StaticClass(),
+		UMovieSceneMirroringDecoration::StaticClass());
 }
 
 void FMovieSceneMixedControlRigModule::ShutdownModule()
 {
-
+	if (IMovieSceneModule::IsAvailable())
+	{
+		IMovieSceneModule::Get().UnregisterCompatibleDecoration(
+			UMovieSceneControlRigParameterTrack::StaticClass(),
+			UMovieSceneMirroringDecoration::StaticClass());
+	}
 }
 
 }
