Sequencer: Fix root motion in the anim mixer in editor in the case that the anims have 'root motion lock' selected.

#Synced-CL 49622961

diff --git a/Engine/Plugins/Experimental/MovieSceneAnimMixer/Source/MovieSceneAnimMixer/Private/Systems/MovieSceneMixedSkeletalAnimationSystem.cpp b/Engine/Plugins/Experimental/MovieSceneAnimMixer/Source/MovieSceneAnimMixer/Private/Systems/MovieSceneMixedSkeletalAnimationSystem.cpp
--- a/Engine/Plugins/Experimental/MovieSceneAnimMixer/Source/MovieSceneAnimMixer/Private/Systems/MovieSceneMixedSkeletalAnimationSystem.cpp
+++ b/Engine/Plugins/Experimental/MovieSceneAnimMixer/Source/MovieSceneAnimMixer/Private/Systems/MovieSceneMixedSkeletalAnimationSystem.cpp
@@ -459,7 +459,15 @@ void FMovieSceneSkeletalAnimationEvaluationTask::Execute(UE::UAF::FEvaluationVM&
 			// Extract the root motion if we've chosen that option
 			if (bUseRootMotion)
 			{
-				RootBoneTransform = Keyframe.Pose.LocalTransformsView[RootIndex];
+#if WITH_EDITOR
+				// If we've locked roots, the root won't hold any transform. We'll extract it later below.
+				if (!AnimSequencePtr->bForceRootLock)
+				{
+#endif
+					RootBoneTransform = Keyframe.Pose.LocalTransformsView[RootIndex];
+#if WITH_EDITOR
+				}
+#endif
 			}
 		}
 
@@ -491,7 +499,11 @@ void FMovieSceneSkeletalAnimationEvaluationTask::Execute(UE::UAF::FEvaluationVM&
 				}
 				else if (bUseRootMotion)
 				{
-					RootBoneTransform = AnimSequencePtr->ExtractRootTrackTransform(FAnimExtractContext(AnimationData.ToPosition), nullptr);
+					FAnimExtractContext ExtractContext(AnimationData.ToPosition);
+#if WITH_EDITOR
+					ExtractContext.bIgnoreRootLock = true;
+#endif
+					RootBoneTransform = AnimSequencePtr->ExtractRootTrackTransform(ExtractContext, nullptr);
 				}
 			}
 		}
