Sequencer- Anim Mixer Transitions
First pass at new 'transition' sections in the animation mixer. These replace easing curves when 2 sections overlap in the same layer with a real section, allowing sections of different types for different transition types with different settings that can be added up the user. The default transition is created when users overlap 2 sections: a 'cross-fade'. I've also included an inertial dead-blend type as an example of another kind of transition.
Transitions are auto-created when 2 sections overlap, and auto-deleted when that overlap goes away. If 3 or more sections overlap transitions are removed and the behavior defaults back to weighted average.

Transitions are implemented in the mixer by implementing an evaluation task that wraps the evaluation task of the 'from' and 'to' sections in the transition. The normal evaluation tasks for the transition sections are suppressed during the transition, giving the transition itself authority over how to evaluate and combine the resulting poses of those two sections.

#Synced-CL 49920339

diff --git a/Engine/Plugins/Experimental/MovieSceneAnimMixer/Source/MovieSceneAnimMixer/Private/AnimMixerComponentTypes.cpp b/Engine/Plugins/Experimental/MovieSceneAnimMixer/Source/MovieSceneAnimMixer/Private/AnimMixerComponentTypes.cpp
--- a/Engine/Plugins/Experimental/MovieSceneAnimMixer/Source/MovieSceneAnimMixer/Private/AnimMixerComponentTypes.cpp
+++ b/Engine/Plugins/Experimental/MovieSceneAnimMixer/Source/MovieSceneAnimMixer/Private/AnimMixerComponentTypes.cpp
@@ -49,9 +49,11 @@ namespace UE::MovieScene
 		ComponentRegistry->NewComponentType(&RootDestination, TEXT("Root Destination"));
 		ComponentRegistry->NewComponentType(&MixerLayer, TEXT("Mixer Layer"));
 		ComponentRegistry->NewComponentType(&MirrorTable, TEXT("Mirror Table"));
+		ComponentRegistry->NewComponentType(&EntityOwner, TEXT("Entity Owner"));
 
 
 		Tags.RequiresBlending = ComponentRegistry->NewTag(TEXT("Requires Blending"));
+		Tags.Transition = ComponentRegistry->NewTag(TEXT("Transition"));
 
 		ComponentRegistry->Factories.DuplicateChildComponent(Priority);
 		ComponentRegistry->Factories.DuplicateChildComponent(Target);
@@ -62,6 +64,7 @@ namespace UE::MovieScene
 		ComponentRegistry->Factories.DuplicateChildComponent(MixerRootMotion);
 		ComponentRegistry->Factories.DuplicateChildComponent(RootDestination);
 		ComponentRegistry->Factories.DuplicateChildComponent(MixerLayer);
+		ComponentRegistry->Factories.DuplicateChildComponent(EntityOwner);
 
 		ComponentRegistry->Factories.DuplicateChildComponent(MirrorTable);
 		
@@ -69,6 +72,7 @@ namespace UE::MovieScene
 		ComponentRegistry->Factories.DefineChildComponent(Task, MixerEntry);
 
 		ComponentRegistry->Factories.DefineChildComponent(Tags.RequiresBlending, Tags.RequiresBlending);
+		ComponentRegistry->Factories.DefineChildComponent(Tags.Transition, Tags.Transition);
 
 	}
 
diff --git a/Engine/Plugins/Experimental/MovieSceneAnimMixer/Source/MovieSceneAnimMixer/Private/EvaluationVM/Tasks/TransitionEvaluationTask.cpp b/Engine/Plugins/Experimental/MovieSceneAnimMixer/Source/MovieSceneAnimMixer/Private/EvaluationVM/Tasks/TransitionEvaluationTask.cpp
new file mode 100644
--- /dev/null
+++ b/Engine/Plugins/Experimental/MovieSceneAnimMixer/Source/MovieSceneAnimMixer/Private/EvaluationVM/Tasks/TransitionEvaluationTask.cpp
@@ -0,0 +1,17 @@
+// Copyright Epic Games, Inc. All Rights Reserved.
+
+#include "EvaluationVM/Tasks/TransitionEvaluationTask.h"
+
+#include UE_INLINE_GENERATED_CPP_BY_NAME(TransitionEvaluationTask)
+
+void FAnimNextTransitionEvaluationTask::Update(
+	const TSharedPtr<FAnimNextEvaluationTask>& InFromTask,
+	const TSharedPtr<FAnimNextEvaluationTask>& InToTask,
+	double InBlendWeight,
+	float InDeltaTime)
+{
+	FromTask = InFromTask;
+	ToTask = InToTask;
+	BlendWeight = InBlendWeight;
+	DeltaTime = InDeltaTime;
+}
diff --git a/Engine/Plugins/Experimental/MovieSceneAnimMixer/Source/MovieSceneAnimMixer/Private/MovieSceneAnimCrossfadeTransitionSection.cpp b/Engine/Plugins/Experimental/MovieSceneAnimMixer/Source/MovieSceneAnimMixer/Private/MovieSceneAnimCrossfadeTransitionSection.cpp
new file mode 100644
--- /dev/null
+++ b/Engine/Plugins/Experimental/MovieSceneAnimMixer/Source/MovieSceneAnimMixer/Private/MovieSceneAnimCrossfadeTransitionSection.cpp
@@ -0,0 +1,110 @@
+// Copyright Epic Games, Inc. All Rights Reserved.
+
+#include "MovieSceneAnimCrossfadeTransitionSection.h"
+
+#include "Channels/MovieSceneChannelProxy.h"
+#include "EvaluationVM/EvaluationProgram.h"
+#include "EvaluationVM/Tasks/ExecuteProgram.h"
+#include "EvaluationVM/Tasks/TransitionEvaluationTask.h"
+#include "Systems/MovieSceneAnimMixerSystem.h"
+
+/**
+ * FAnimNextCrossfadeTransitionTask
+ *
+ * A simple crossfade transition task that blends between two poses
+ * based on a blend weight.
+ */
+struct FAnimNextCrossfadeTransitionTask : public FAnimNextTransitionEvaluationTask
+{
+	DECLARE_ANIM_EVALUATION_TASK(FAnimNextCrossfadeTransitionTask)
+
+	virtual void Execute(UE::UAF::FEvaluationVM& VM) const override
+	{
+		// Handle edge cases where one or both tasks are missing
+		if (!FromTask && !ToTask)
+		{
+			return;
+		}
+		if (!FromTask)
+		{
+			ToTask->Execute(VM);
+			return;
+		}
+		if (!ToTask)
+		{
+			FromTask->Execute(VM);
+			return;
+		}
+
+		// Execute from task first (pushes pose onto stack)
+		FromTask->Execute(VM);
+
+		// Execute to task (pushes second pose onto stack)
+		ToTask->Execute(VM);
+
+		// Blend the two poses using the blend weight
+		// FAnimNextBlendTwoKeyframesPreserveRootMotionTask blends top two stack items
+		FAnimNextBlendTwoKeyframesPreserveRootMotionTask BlendTask =
+			FAnimNextBlendTwoKeyframesPreserveRootMotionTask::Make(static_cast<float>(BlendWeight));
+		BlendTask.Execute(VM);
+	}
+};
+
+UMovieSceneAnimCrossfadeTransitionSection::UMovieSceneAnimCrossfadeTransitionSection(const FObjectInitializer& ObjInit)
+	: Super(ObjInit)
+{
+	// Initialize with a default linear blend curve
+	InitializeDefaultCurve();
+}
+
+void UMovieSceneAnimCrossfadeTransitionSection::InitializeDefaultCurve()
+{
+	BlendCurve.Reset();
+
+	TRange<FFrameNumber> Range = GetRange();
+
+	if (Range.IsEmpty() || !Range.HasLowerBound() || !Range.HasUpperBound())
+	{
+		BlendCurve.SetDefault(0.0f);
+		return;
+	}
+
+	// Add keyframes with cubic interpolation for smooth ease-in/ease-out (like easing curves)
+	FFrameNumber StartFrame = Range.GetLowerBoundValue();
+	FFrameNumber EndFrame = Range.GetUpperBoundValue();
+	BlendCurve.AddCubicKey(StartFrame, 0.0f, RCTM_Auto);
+	BlendCurve.AddCubicKey(EndFrame, 1.0f, RCTM_Auto);
+}
+
+TSharedPtr<FAnimNextTransitionEvaluationTask> UMovieSceneAnimCrossfadeTransitionSection::CreateTransitionTask() const
+{
+	return MakeShared<FAnimNextCrossfadeTransitionTask>();
+}
+
+FName UMovieSceneAnimCrossfadeTransitionSection::GetTransitionIconStyleName() const
+{
+	// TODO: Register a custom icon for crossfade transitions
+	// For now, return NAME_None to skip icon drawing
+	return NAME_None;
+}
+
+FText UMovieSceneAnimCrossfadeTransitionSection::GetTransitionDisplayName() const
+{
+	return NSLOCTEXT("MovieSceneAnimCrossfadeTransitionSection", "CrossfadeDisplayName", "Crossfade");
+}
+
+void UMovieSceneAnimCrossfadeTransitionSection::RebuildChannelProxy(FMovieSceneChannelProxyData& Channels)
+{
+#if WITH_EDITOR
+	FMovieSceneChannelMetaData MetaData(TEXT("BlendCurve"), NSLOCTEXT("MovieSceneAnimCrossfadeTransitionSection", "BlendCurve", "Blend"));
+	// Don't collapse to track - show as expandable row so users can access the blend curve
+	MetaData.bCanCollapseToTrack = false;
+	Channels.Add(
+		BlendCurve,
+		MetaData,
+		TMovieSceneExternalValue<float>()
+	);
+#else
+	Channels.Add(BlendCurve);
+#endif
+}
diff --git a/Engine/Plugins/Experimental/MovieSceneAnimMixer/Source/MovieSceneAnimMixer/Private/MovieSceneAnimInertialDeadBlendTransitionSection.cpp b/Engine/Plugins/Experimental/MovieSceneAnimMixer/Source/MovieSceneAnimMixer/Private/MovieSceneAnimInertialDeadBlendTransitionSection.cpp
new file mode 100644
--- /dev/null
+++ b/Engine/Plugins/Experimental/MovieSceneAnimMixer/Source/MovieSceneAnimMixer/Private/MovieSceneAnimInertialDeadBlendTransitionSection.cpp
@@ -0,0 +1,249 @@
+// Copyright Epic Games, Inc. All Rights Reserved.
+
+#include "MovieSceneAnimInertialDeadBlendTransitionSection.h"
+
+#include "Channels/MovieSceneChannelProxy.h"
+#include "EvaluationVM/EvaluationTask.h"
+#include "EvaluationVM/EvaluationVM.h"
+#include "EvaluationVM/Tasks/TransitionEvaluationTask.h"
+#include "EvaluationVM/Tasks/DeadBlending.h"
+#include "EvaluationVM/Tasks/StoreKeyframe.h"
+#include "Traits/Inertialization.h"
+#include "TransformArray.h"
+#include "TransformArrayOperations.h"
+#include "Curves/CurveFloat.h"
+#include "MovieScene.h"
+
+/**
+ * FAnimNextInertialDeadBlendTransitionTask
+ *
+ * A stateful transition task that performs inertial dead blending.
+ * Captures poses from the "from" section to compute velocity, then
+ * extrapolates that motion while blending toward the "to" pose.
+ * Uses the pre-existing UAF Evaluation Tasks internally for the actual dead blend calculations.
+ *
+ * Execution flow:
+ * - Frame 1: Capture from-pose, output from-pose (need velocity data, avoid pop)
+ * - Frame 2: Capture second from-pose, compute velocity, initialize dead blend, apply
+ * - Subsequent frames: Apply dead blending to extrapolate and blend
+ */
+struct FAnimNextInertialDeadBlendTransitionTask : public FAnimNextTransitionEvaluationTask
+{
+	DECLARE_ANIM_EVALUATION_TASK(FAnimNextInertialDeadBlendTransitionTask)
+
+	// Parameters (set at creation from section properties)
+	float BlendDuration = 0.3f;
+	EAlphaBlendOption BlendModeParam = EAlphaBlendOption::HermiteCubic;
+	TWeakObjectPtr<UCurveFloat> CustomBlendCurveParam = nullptr;
+	UE::UAF::FDeadBlendTransitionTaskParameters ExtrapolationParams;
+
+	// Persistent state
+	mutable UE::UAF::FTransformArraySoAHeap CurrFromPose;
+	mutable UE::UAF::FTransformArraySoAHeap PrevFromPose;
+	mutable UE::UAF::FDeadBlendingState DeadBlendState;
+	mutable float PoseDeltaTime = 0.0f;
+	mutable float TimeSinceTransitionStart = 0.0f;
+	mutable bool bHasCurrPose = false;
+	mutable bool bHasPrevPose = false;
+	mutable bool bTransitionInitialized = false;
+
+	virtual void Update(
+		const TSharedPtr<FAnimNextEvaluationTask>& InFromTask,
+		const TSharedPtr<FAnimNextEvaluationTask>& InToTask,
+		double InBlendWeight,
+		float InDeltaTime) override
+	{
+		// Call base to update from/to tasks and blend weight
+		FAnimNextTransitionEvaluationTask::Update(InFromTask, InToTask, InBlendWeight, InDeltaTime);
+
+		// Accumulate time since transition started
+		if (bTransitionInitialized)
+		{
+			TimeSinceTransitionStart += InDeltaTime;
+		}
+
+		// Accumulate delta time for pose velocity computation
+		PoseDeltaTime += InDeltaTime;
+	}
+
+	virtual void Execute(UE::UAF::FEvaluationVM& VM) const override
+	{
+		using namespace UE::UAF;
+
+		// Handle edge cases
+		if (!FromTask && !ToTask)
+		{
+			return;
+		}
+		if (!FromTask)
+		{
+			ToTask->Execute(VM);
+			return;
+		}
+		if (!ToTask)
+		{
+			FromTask->Execute(VM);
+			return;
+		}
+
+		// Phase 1: First frame - capture from-pose, output from-pose (avoid pop)
+		if (!bHasCurrPose)
+		{
+			// Execute from-task to get from-pose
+			FromTask->Execute(VM);
+
+			// Copy from stack to our storage (don't pop - leave it as output)
+			if (const TUniquePtr<FKeyframeState>* Keyframe = VM.PeekValue<TUniquePtr<FKeyframeState>>(KEYFRAME_STACK_NAME, 0))
+			{
+				if (*Keyframe)
+				{
+					CurrFromPose.SetNumUninitialized((*Keyframe)->Pose.LocalTransforms.Num());
+					CopyTransforms(CurrFromPose.GetView(), (*Keyframe)->Pose.LocalTransforms.GetConstView());
+					bHasCurrPose = true;
+				}
+			}
+
+			// Reset delta time for next capture
+			PoseDeltaTime = 0.0f;
+
+			// Leave from-pose on stack as output (no pop)
+			return;
+		}
+
+		// Phase 2: Second frame - we have CurrPose, capture another for velocity
+		if (!bHasPrevPose && PoseDeltaTime > UE_SMALL_NUMBER)
+		{
+			// Swap poses - move current to previous
+			PrevFromPose = MoveTemp(CurrFromPose);
+
+			// Execute from-task to capture new current pose
+			FromTask->Execute(VM);
+
+			const TUniquePtr<FKeyframeState>* Keyframe = VM.PeekValue<TUniquePtr<FKeyframeState>>(KEYFRAME_STACK_NAME, 0);
+			if (!Keyframe || !(*Keyframe))
+			{
+				// Failed to capture pose - pop any stack entry and restore previous state
+				TUniquePtr<FKeyframeState> DiscardedKeyframe;
+				(void)VM.PopValue(KEYFRAME_STACK_NAME, DiscardedKeyframe);
+				CurrFromPose = MoveTemp(PrevFromPose);
+				return;
+			}
+
+			CurrFromPose.SetNumUninitialized((*Keyframe)->Pose.LocalTransforms.Num());
+			CopyTransforms(CurrFromPose.GetView(), (*Keyframe)->Pose.LocalTransforms.GetConstView());
+
+			// Pop the from-pose (discard it)
+			TUniquePtr<FKeyframeState> PoppedKeyframe;
+			(void)VM.PopValue(KEYFRAME_STACK_NAME, PoppedKeyframe);
+
+			bHasPrevPose = true;
+
+			// Now execute to-task to get destination pose
+			ToTask->Execute(VM);
+
+			// Initialize dead blending state with velocity
+			if (const TUniquePtr<FKeyframeState>* ToKeyframe = VM.PeekValue<TUniquePtr<FKeyframeState>>(KEYFRAME_STACK_NAME, 0))
+			{
+				// Use the existing transition task to initialize our state
+				FAnimNextDeadBlendingTransitionTask TransitionTask = FAnimNextDeadBlendingTransitionTask::Make(
+					&DeadBlendState,
+					&CurrFromPose,
+					&PrevFromPose,
+					PoseDeltaTime,
+					ExtrapolationParams);
+
+				// Execute transition task - it will read destination from stack and fill DeadBlendState
+				TransitionTask.Execute(VM);
+			}
+
+			bTransitionInitialized = true;
+			TimeSinceTransitionStart = 0.0f;
+			PoseDeltaTime = 0.0f;
+
+			// Apply dead blending to the to-pose on the stack
+			{
+				FAnimNextDeadBlendingApplyTask ApplyTask = FAnimNextDeadBlendingApplyTask::Make(
+					&DeadBlendState,
+					BlendDuration,
+					TimeSinceTransitionStart,
+					BlendModeParam,
+					CustomBlendCurveParam.Get());
+
+				ApplyTask.Execute(VM);
+			}
+
+			return;
+		}
+
+		// Phase 3: Subsequent frames - apply dead blending
+		if (bTransitionInitialized)
+		{
+			// Execute to-task to get destination pose
+			ToTask->Execute(VM);
+
+			// Apply dead blending
+			FAnimNextDeadBlendingApplyTask ApplyTask = FAnimNextDeadBlendingApplyTask::Make(
+				&DeadBlendState,
+				BlendDuration,
+				TimeSinceTransitionStart,
+				BlendModeParam,
+				CustomBlendCurveParam.Get());
+
+			ApplyTask.Execute(VM);
+		}
+	}
+};
+
+UMovieSceneAnimInertialDeadBlendTransitionSection::UMovieSceneAnimInertialDeadBlendTransitionSection(const FObjectInitializer& ObjInit)
+	: Super(ObjInit)
+{
+}
+
+TSharedPtr<FAnimNextTransitionEvaluationTask> UMovieSceneAnimInertialDeadBlendTransitionSection::CreateTransitionTask() const
+{
+	UMovieScene* MovieScene = GetTypedOuter<UMovieScene>();
+	if (!MovieScene)
+	{
+		return nullptr;
+	}
+	TSharedPtr<FAnimNextInertialDeadBlendTransitionTask> Task = MakeShared<FAnimNextInertialDeadBlendTransitionTask>();
+
+	// Calculate blend duration from section range
+	TRange<FFrameNumber> Range = GetRange();
+	if (Range.HasLowerBound() && Range.HasUpperBound())
+	{
+		FFrameNumber Duration = Range.GetUpperBoundValue() - Range.GetLowerBoundValue();
+		// Get frame rate from the owning movie scene
+		FFrameRate TickResolution = MovieScene->GetTickResolution();
+		Task->BlendDuration = static_cast<float>(Duration.Value) / static_cast<float>(TickResolution.AsDecimal());
+	}
+
+	// Configure task with section parameters
+	Task->BlendModeParam = BlendMode;
+	Task->CustomBlendCurveParam = CustomBlendCurve;
+
+	Task->ExtrapolationParams.ExtrapolationHalfLife = ExtrapolationHalfLife;
+	Task->ExtrapolationParams.ExtrapolationHalfLifeMin = ExtrapolationHalfLifeMin;
+	Task->ExtrapolationParams.ExtrapolationHalfLifeMax = ExtrapolationHalfLifeMax;
+	Task->ExtrapolationParams.MaximumTranslationVelocity = MaximumTranslationVelocity;
+	Task->ExtrapolationParams.MaximumRotationVelocity = FMath::DegreesToRadians(MaximumRotationVelocity);
+	Task->ExtrapolationParams.MaximumScaleVelocity = MaximumScaleVelocity;
+
+	return Task;
+}
+
+FName UMovieSceneAnimInertialDeadBlendTransitionSection::GetTransitionIconStyleName() const
+{
+	// TODO: Make an icon
+	return NAME_None;
+}
+
+FText UMovieSceneAnimInertialDeadBlendTransitionSection::GetTransitionDisplayName() const
+{
+	return NSLOCTEXT("MovieSceneAnimInertialDeadBlendTransitionSection", "DisplayName", "Inertial Dead Blend");
+}
+
+void UMovieSceneAnimInertialDeadBlendTransitionSection::RebuildChannelProxy(FMovieSceneChannelProxyData& Channels)
+{
+	// No channels needed - dead blending doesn't use a blend curve
+}
diff --git a/Engine/Plugins/Experimental/MovieSceneAnimMixer/Source/MovieSceneAnimMixer/Private/MovieSceneAnimTransitionSectionBase.cpp b/Engine/Plugins/Experimental/MovieSceneAnimMixer/Source/MovieSceneAnimMixer/Private/MovieSceneAnimTransitionSectionBase.cpp
new file mode 100644
--- /dev/null
+++ b/Engine/Plugins/Experimental/MovieSceneAnimMixer/Source/MovieSceneAnimMixer/Private/MovieSceneAnimTransitionSectionBase.cpp
@@ -0,0 +1,246 @@
+// Copyright Epic Games, Inc. All Rights Reserved.
+
+#include "MovieSceneAnimTransitionSectionBase.h"
+
+#include "AnimMixerComponentTypes.h"
+#include "MovieSceneAnimationMixerTrack.h"
+#include "Channels/MovieSceneChannelProxy.h"
+#include "EntitySystem/BuiltInComponentTypes.h"
+#include "Evaluation/MovieSceneEvaluationField.h"
+#include "MovieSceneAnimationMixerLayer.h"
+#include "EvaluationVM/EvaluationTask.h"
+#include "EvaluationVM/Tasks/TransitionEvaluationTask.h"
+#include "Systems/MovieSceneSkeletalAnimationSystem.h"
+
+UMovieSceneAnimTransitionSectionBase::UMovieSceneAnimTransitionSectionBase(const FObjectInitializer& ObjInit)
+	: Super(ObjInit)
+{
+}
+
+bool UMovieSceneAnimTransitionSectionBase::IsValid() const
+{
+	if (!FromSection || !ToSection)
+	{
+		return false;
+	}
+
+	// All three sections must be on the same row
+	const int32 TransitionRow = GetRowIndex();
+	if (FromSection->GetRowIndex() != TransitionRow || ToSection->GetRowIndex() != TransitionRow)
+	{
+		return false;
+	}
+
+	TRange<FFrameNumber> FromRange = FromSection->GetRange();
+	TRange<FFrameNumber> ToRange = ToSection->GetRange();
+
+	// Both sections must have finite bounds
+	if (!FromRange.HasLowerBound() || !FromRange.HasUpperBound() ||
+		!ToRange.HasLowerBound() || !ToRange.HasUpperBound())
+	{
+		return false;
+	}
+
+	// Check if one section is fully contained within the other - this is not a valid transition
+	// A transition requires both sections to extend beyond the overlap region
+	// If one section fully contains the other, there's no valid transition point
+	if (FromRange.Contains(ToRange) || ToRange.Contains(FromRange))
+	{
+		return false;
+	}
+
+	// There must be an overlap between the sections
+	TRange<FFrameNumber> Overlap = TRange<FFrameNumber>::Intersection(FromRange, ToRange);
+	if (Overlap.IsEmpty())
+	{
+		return false;
+	}
+
+	return true;
+}
+
+TRange<FFrameNumber> UMovieSceneAnimTransitionSectionBase::ComputeOverlapRange() const
+{
+	if (!IsValid())
+	{
+		return TRange<FFrameNumber>::Empty();
+	}
+
+	return TRange<FFrameNumber>::Intersection(
+		FromSection->GetRange(),
+		ToSection->GetRange()
+	);
+}
+
+bool UMovieSceneAnimTransitionSectionBase::UpdateBoundsFromSourceSections()
+{
+	// IsValid() checks: sections exist, same row, bounded, no containment, has overlap
+	if (!IsValid())
+	{
+		return false;
+	}
+
+	TRange<FFrameNumber> FromRange = FromSection->GetRange();
+	TRange<FFrameNumber> ToRange = ToSection->GetRange();
+
+	// Check if the order has reversed (the "to" section now starts before "from")
+	// In a proper transition, FromSection should start before ToSection
+	if (ToRange.GetLowerBoundValue() < FromRange.GetLowerBoundValue())
+	{
+		// Swap from/to to maintain correct ordering
+		Modify();
+		Swap(FromSection, ToSection);
+
+		// Update local variables after swap
+		FromRange = FromSection->GetRange();
+		ToRange = ToSection->GetRange();
+	}
+
+	TRange<FFrameNumber> NewOverlap = TRange<FFrameNumber>::Intersection(FromRange, ToRange);
+
+	// Check if more than 2 sections overlap this range - transitions only work between exactly 2 sections
+	if (UMovieSceneAnimationMixerTrack* MixerTrack = GetTypedOuter<UMovieSceneAnimationMixerTrack>())
+	{
+		if (MixerTrack->CountSectionsOverlappingRange(NewOverlap, GetRowIndex()) > 2)
+		{
+			return false;
+		}
+	}
+
+	// Get the current range before updating
+	TRange<FFrameNumber> OldRange = GetRange();
+
+	// Update transition bounds to match the full overlap
+	SetRange(NewOverlap);
+
+	// Scale existing keys proportionally to the new range
+	ScaleKeysToRange(OldRange, NewOverlap);
+
+	return true;
+}
+
+void UMovieSceneAnimTransitionSectionBase::ScaleKeysToRange(const TRange<FFrameNumber>& OldRange, const TRange<FFrameNumber>& NewRange)
+{
+	const FMovieSceneFloatChannel* ConstChannel = GetBlendWeightChannel();
+	if (!ConstChannel)
+	{
+		return;
+	}
+
+	if (!OldRange.HasLowerBound() || !OldRange.HasUpperBound() ||
+		!NewRange.HasLowerBound() || !NewRange.HasUpperBound())
+	{
+		return;
+	}
+
+	const FFrameNumber OldStart = OldRange.GetLowerBoundValue();
+	const FFrameNumber OldEnd = OldRange.GetUpperBoundValue();
+	const FFrameNumber NewStart = NewRange.GetLowerBoundValue();
+	const FFrameNumber NewEnd = NewRange.GetUpperBoundValue();
+
+	const double OldDuration = (OldEnd - OldStart).Value;
+	const double NewDuration = (NewEnd - NewStart).Value;
+
+	if (OldDuration <= 0 || NewDuration <= 0)
+	{
+		return;
+	}
+
+	// Need non-const access to modify keys
+	FMovieSceneFloatChannel* Channel = const_cast<FMovieSceneFloatChannel*>(ConstChannel);
+
+	// Collect existing key data
+	TArrayView<const FFrameNumber> OldTimes = Channel->GetData().GetTimes();
+	TArrayView<const FMovieSceneFloatValue> OldValues = Channel->GetValues();
+
+	TArray<FFrameNumber> NewTimes;
+	TArray<FMovieSceneFloatValue> NewValues;
+	NewTimes.Reserve(OldTimes.Num());
+	NewValues.Reserve(OldValues.Num());
+
+	for (int32 i = 0; i < OldTimes.Num(); ++i)
+	{
+		// Calculate normalized position in old range [0, 1]
+		double NormalizedPos = (OldTimes[i] - OldStart).Value / OldDuration;
+		// Map to new range
+		FFrameNumber ScaledTime = NewStart + FFrameNumber(static_cast<int32>(NormalizedPos * NewDuration));
+		NewTimes.Add(ScaledTime);
+		NewValues.Add(OldValues[i]);
+	}
+
+	// Clear and re-add keys at new positions
+	Channel->Reset();
+	Channel->AddKeys(NewTimes, NewValues);
+}
+
+float UMovieSceneAnimTransitionSectionBase::GetTransitionProgress(FFrameTime InTime) const
+{
+	TRange<FFrameNumber> Range = GetRange();
+	if (Range.IsEmpty() || !Range.HasLowerBound() || !Range.HasUpperBound())
+	{
+		return 0.0f;
+	}
+
+	FFrameNumber Start = Range.GetLowerBoundValue();
+	FFrameNumber End = Range.GetUpperBoundValue();
+	FFrameNumber Duration = End - Start;
+
+	if (Duration.Value <= 0)
+	{
+		return 0.0f;
+	}
+
+	double Progress = (InTime.AsDecimal() - Start.Value) / Duration.Value;
+	return FMath::Clamp(static_cast<float>(Progress), 0.0f, 1.0f);
+}
+
+EMovieSceneChannelProxyType UMovieSceneAnimTransitionSectionBase::CacheChannelProxy()
+{
+	FMovieSceneChannelProxyData Channels;
+
+	// Let subclass add its specific channels
+	RebuildChannelProxy(Channels);
+
+	ChannelProxy = MakeShared<FMovieSceneChannelProxy>(MoveTemp(Channels));
+	return EMovieSceneChannelProxyType::Static;
+}
+
+void UMovieSceneAnimTransitionSectionBase::ImportEntityImpl(
+	UMovieSceneEntitySystemLinker* EntityLinker,
+	const FEntityImportParams& Params,
+	FImportedEntity* OutImportedEntity)
+{
+	using namespace UE::MovieScene;
+
+	const FBuiltInComponentTypes* BuiltInComponents = FBuiltInComponentTypes::Get();
+	const FAnimMixerComponentTypes* AnimMixerComponents = FAnimMixerComponentTypes::Get();
+
+	if (Params.GetObjectBindingID().IsValid())
+	{
+		// Create the transition task - it will persist as an entity component and have
+		// its Update() method called each frame by the mixer system to set from/to tasks.
+		// The Transition tag marks this entity for special handling by the mixer system.
+		TSharedPtr<FAnimNextTransitionEvaluationTask> TransitionTask = CreateTransitionTask();
+		const FMovieSceneFloatChannel* BlendChannel = GetBlendWeightChannel();
+
+		OutImportedEntity->AddBuilder(
+			FEntityBuilder()
+			.Add(BuiltInComponents->GenericObjectBinding, Params.GetObjectBindingID())
+			.Add(BuiltInComponents->BoundObjectResolver, UMovieSceneSkeletalAnimationSystem::ResolveSkeletalMeshComponentBinding)
+			.Add(AnimMixerComponents->Task, TransitionTask)
+			.Add(AnimMixerComponents->EntityOwner, FObjectKey(this))
+			.AddTag(AnimMixerComponents->Tags.Transition)
+			.AddConditional(BuiltInComponents->WeightChannel, BlendChannel, BlendChannel != nullptr)
+		);
+	}
+}
+
+bool UMovieSceneAnimTransitionSectionBase::PopulateEvaluationFieldImpl(
+	const TRange<FFrameNumber>& EffectiveRange,
+	const FMovieSceneEvaluationFieldEntityMetaData& InMetaData,
+	FMovieSceneEntityComponentFieldBuilder* OutFieldBuilder)
+{
+	const int32 MetaDataIndex = OutFieldBuilder->AddMetaData(InMetaData);
+	OutFieldBuilder->AddPersistentEntity(EffectiveRange, this, 0, MetaDataIndex);
+	return true;
+}
diff --git a/Engine/Plugins/Experimental/MovieSceneAnimMixer/Source/MovieSceneAnimMixer/Private/MovieSceneAnimationMixerLayer.cpp b/Engine/Plugins/Experimental/MovieSceneAnimMixer/Source/MovieSceneAnimMixer/Private/MovieSceneAnimationMixerLayer.cpp
--- a/Engine/Plugins/Experimental/MovieSceneAnimMixer/Source/MovieSceneAnimMixer/Private/MovieSceneAnimationMixerLayer.cpp
+++ b/Engine/Plugins/Experimental/MovieSceneAnimMixer/Source/MovieSceneAnimMixer/Private/MovieSceneAnimationMixerLayer.cpp
@@ -14,6 +14,8 @@ void UMovieSceneAnimationMixerLayer::SetChildTrack(UMovieSceneTrack* InTrack)
 {
 	if (ChildTrack != InTrack)
 	{
+		Modify();
+
 		// Clear sections if we're setting a child track
 		if (InTrack && Sections.Num() > 0)
 		{
@@ -29,6 +31,8 @@ void UMovieSceneAnimationMixerLayer::AddSection(UMovieSceneSection* Section)
 {
 	if (Section && !Sections.Contains(Section))
 	{
+		Modify();
+
 		// Clear child track if we're adding sections
 		if (ChildTrack)
 		{
@@ -42,8 +46,10 @@ void UMovieSceneAnimationMixerLayer::AddSection(UMovieSceneSection* Section)
 
 void UMovieSceneAnimationMixerLayer::RemoveSection(UMovieSceneSection* Section)
 {
-	if (Sections.Remove(Section) > 0)
+	if (Sections.Contains(Section))
 	{
+		Modify();
+		Sections.Remove(Section);
 		MarkAsChanged();
 	}
 }
diff --git a/Engine/Plugins/Experimental/MovieSceneAnimMixer/Source/MovieSceneAnimMixer/Private/MovieSceneAnimationMixerTrack.cpp b/Engine/Plugins/Experimental/MovieSceneAnimMixer/Source/MovieSceneAnimMixer/Private/MovieSceneAnimationMixerTrack.cpp
--- a/Engine/Plugins/Experimental/MovieSceneAnimMixer/Source/MovieSceneAnimMixer/Private/MovieSceneAnimationMixerTrack.cpp
+++ b/Engine/Plugins/Experimental/MovieSceneAnimMixer/Source/MovieSceneAnimMixer/Private/MovieSceneAnimationMixerTrack.cpp
@@ -1,6 +1,8 @@
 // Copyright Epic Games, Inc. All Rights Reserved.
 
 #include "MovieSceneAnimationMixerTrack.h"
+#include "MovieSceneAnimTransitionSectionBase.h"
+#include "MovieSceneAnimCrossfadeTransitionSection.h"
 #include "MovieSceneSection.h"
 #include "Sections/MovieSceneSkeletalAnimationSection.h"
 #include "Channels/MovieSceneChannelProxy.h"
@@ -210,6 +212,7 @@ void UMovieSceneAnimationSectionDecoration::ExtendEntityImpl(UMovieSceneEntitySy
 		.Add(AnimMixerComponents->Target, Target)
 		.Add(AnimMixerComponents->Priority, MixerLayer ? MixerLayer->GetLayerIndex() : 0)
 		.Add(AnimMixerComponents->MixerLayer, MixerLayer)
+		.Add(AnimMixerComponents->EntityOwner, FObjectKey(Section))
 		.AddMutualComponents()
 	);
 }
@@ -239,6 +242,18 @@ bool UMovieSceneAnimationMixerTrack::FixRowIndices()
 	return false;
 }
 
+void UMovieSceneAnimationMixerTrack::RemoveSection(UMovieSceneSection& Section)
+{
+	Modify();
+
+	// Call OnSectionRemovedImpl FIRST so we can find and cascade-delete related transitions
+	// before the section is removed from our internal arrays
+	OnSectionRemovedImpl(&Section);
+
+	// Now call parent to actually remove the section
+	Super::RemoveSection(Section);
+}
+
 void UMovieSceneAnimationMixerTrack::OnSectionAddedImpl(UMovieSceneSection* Section)
 {
 	if (Section)
@@ -267,6 +282,11 @@ void UMovieSceneAnimationMixerTrack::OnSectionAddedImpl(UMovieSceneSection* Sect
 		{
 			Layer->AddSection(Section);
 		}
+
+		if (!Section->IsA<UMovieSceneAnimTransitionSectionBase>())
+		{
+			UpdateTransitionsForRow(RowIndex);
+		}
 	}
 }
 
@@ -280,9 +300,440 @@ void UMovieSceneAnimationMixerTrack::OnSectionRemovedImpl(UMovieSceneSection* Se
 		{
 			Layer->RemoveSection(Section);
 		}
+
+		// Cascade delete: remove any transition sections that reference this section
+		TArray<UMovieSceneAnimTransitionSectionBase*> RelatedTransitions = FindTransitionsForSection(Section);
+		for (UMovieSceneAnimTransitionSectionBase* Transition : RelatedTransitions)
+		{
+			RemoveSection(*Transition);
+		}
+	}
+}
+
+TArray<UMovieSceneAnimTransitionSectionBase*> UMovieSceneAnimationMixerTrack::FindTransitionsForSection(UMovieSceneSection* Section) const
+{
+	TArray<UMovieSceneAnimTransitionSectionBase*> Result;
+
+	if (!Section)
+	{
+		return Result;
+	}
+
+	for (UMovieSceneSection* Sec : GetAllSections())
+	{
+		if (UMovieSceneAnimTransitionSectionBase* Transition = Cast<UMovieSceneAnimTransitionSectionBase>(Sec))
+		{
+			if (Transition->FromSection == Section || Transition->ToSection == Section)
+			{
+				Result.Add(Transition);
+			}
+		}
+	}
+
+	return Result;
+}
+
+TArray<UMovieSceneSection*> UMovieSceneAnimationMixerTrack::FindOverlappingSections(UMovieSceneSection* Section) const
+{
+	TArray<UMovieSceneSection*> Result;
+
+	if (!Section)
+	{
+		return Result;
+	}
+
+	const int32 RowIndex = Section->GetRowIndex();
+	const TRange<FFrameNumber> SectionRange = Section->GetRange();
+
+	if (UMovieSceneAnimationMixerLayer* Layer = GetLayer(RowIndex))
+	{
+		for (UMovieSceneSection* OtherSection : Layer->GetSections())
+		{
+			// Skip self and transitions
+			if (OtherSection == Section ||
+				OtherSection->IsA<UMovieSceneAnimTransitionSectionBase>())
+			{
+				continue;
+			}
+
+			// Check for overlap
+			TRange<FFrameNumber> Overlap = TRange<FFrameNumber>::Intersection(SectionRange, OtherSection->GetRange());
+			if (!Overlap.IsEmpty())
+			{
+				Result.Add(OtherSection);
+			}
+		}
 	}
+
+	return Result;
 }
 
+int32 UMovieSceneAnimationMixerTrack::CountSectionsOverlappingRange(const TRange<FFrameNumber>& Range, int32 RowIndex) const
+{
+	int32 Count = 0;
+
+	if (UMovieSceneAnimationMixerLayer* Layer = GetLayer(RowIndex))
+	{
+		for (UMovieSceneSection* Section : Layer->GetSections())
+		{
+			// Skip transitions
+			if (Section->IsA<UMovieSceneAnimTransitionSectionBase>())
+			{
+				continue;
+			}
+
+			// Check for overlap
+			TRange<FFrameNumber> Overlap = TRange<FFrameNumber>::Intersection(Range, Section->GetRange());
+			if (!Overlap.IsEmpty())
+			{
+				Count++;
+			}
+		}
+	}
+
+	return Count;
+}
+
+UMovieSceneAnimTransitionSectionBase* UMovieSceneAnimationMixerTrack::FindTransitionForPair(UMovieSceneSection* FromSection, UMovieSceneSection* ToSection) const
+{
+	if (!FromSection || !ToSection)
+	{
+		return nullptr;
+	}
+
+	for (UMovieSceneSection* Sec : GetAllSections())
+	{
+		if (UMovieSceneAnimTransitionSectionBase* Transition = Cast<UMovieSceneAnimTransitionSectionBase>(Sec))
+		{
+			if (Transition->FromSection == FromSection && Transition->ToSection == ToSection)
+			{
+				return Transition;
+			}
+		}
+	}
+
+	return nullptr;
+}
+
+UMovieSceneAnimTransitionSectionBase* UMovieSceneAnimationMixerTrack::CreateTransitionSectionOfType(
+	UMovieSceneSection* FromSection,
+	UMovieSceneSection* ToSection,
+	TSubclassOf<UMovieSceneAnimTransitionSectionBase> TransitionClass)
+{
+	if (!FromSection || !ToSection || !TransitionClass)
+	{
+		return nullptr;
+	}
+
+	// Verify the class is valid
+	if (TransitionClass->HasAnyClassFlags(CLASS_Abstract))
+	{
+		return nullptr;
+	}
+
+	// Sections must be on the same row (layer)
+	if (FromSection->GetRowIndex() != ToSection->GetRowIndex())
+	{
+		return nullptr;
+	}
+
+	// Compute the overlap range
+	TRange<FFrameNumber> Overlap = TRange<FFrameNumber>::Intersection(
+		FromSection->GetRange(),
+		ToSection->GetRange()
+	);
+
+	if (Overlap.IsEmpty())
+	{
+		return nullptr;
+	}
+
+	Modify();
+
+	// Create the transition section
+	UMovieSceneAnimTransitionSectionBase* TransitionSection = NewObject<UMovieSceneAnimTransitionSectionBase>(
+		this, TransitionClass, NAME_None, RF_Transactional);
+	TransitionSection->SetRange(Overlap);
+	TransitionSection->SetRowIndex(FromSection->GetRowIndex());
+	TransitionSection->FromSection = FromSection;
+	TransitionSection->ToSection = ToSection;
+
+	// Ensure transition sections always draw on top of regular sections
+	TransitionSection->SetOverlapPriority(INT32_MAX);
+
+	// Initialize any default curves now that the range is set
+	TransitionSection->InitializeDefaultCurve();
+
+	// Add via FSectionParameter to ensure proper lifecycle hooks are called
+	UMovieSceneTrack::AddSection(FSectionParameter{*TransitionSection});
+
+	return TransitionSection;
+}
+
+void UMovieSceneAnimationMixerTrack::AutoCreateTransitionsForSection(UMovieSceneSection* Section)
+{
+	if (!Section || Section->IsA<UMovieSceneAnimTransitionSectionBase>())
+	{
+		return;
+	}
+
+	TArray<UMovieSceneSection*> OverlappingSections = FindOverlappingSections(Section);
+
+	for (UMovieSceneSection* OtherSection : OverlappingSections)
+	{
+		TRange<FFrameNumber> SectionRange = Section->GetRange();
+		TRange<FFrameNumber> OtherRange = OtherSection->GetRange();
+
+		// Skip if either section doesn't have proper bounds
+		if (!SectionRange.HasLowerBound() || !SectionRange.HasUpperBound() ||
+			!OtherRange.HasLowerBound() || !OtherRange.HasUpperBound())
+		{
+			continue;
+		}
+
+		FFrameNumber SectionStart = SectionRange.GetLowerBoundValue();
+		FFrameNumber SectionEnd = SectionRange.GetUpperBoundValue();
+		FFrameNumber OtherStart = OtherRange.GetLowerBoundValue();
+		FFrameNumber OtherEnd = OtherRange.GetUpperBoundValue();
+
+		// Check if one section is fully contained within the other - this is not a valid transition
+		// A transition requires both sections to extend beyond the overlap region
+		bool bSectionInsideOther = (OtherStart <= SectionStart && SectionEnd <= OtherEnd);
+		bool bOtherInsideSection = (SectionStart <= OtherStart && OtherEnd <= SectionEnd);
+
+		if (bSectionInsideOther || bOtherInsideSection)
+		{
+			// One section is fully encompassed by the other - don't create a transition
+			continue;
+		}
+
+		// Compute the overlap range between these two sections
+		TRange<FFrameNumber> OverlapRange = TRange<FFrameNumber>::Intersection(SectionRange, OtherRange);
+
+		// Check if more than 2 sections overlap this range - transitions only work between exactly 2 sections
+		if (CountSectionsOverlappingRange(OverlapRange, Section->GetRowIndex()) > 2)
+		{
+			continue;
+		}
+
+		// Determine which is "from" and which is "to" based on start time
+		UMovieSceneSection* FromSection = nullptr;
+		UMovieSceneSection* ToSection = nullptr;
+
+		if (SectionStart <= OtherStart)
+		{
+			FromSection = Section;
+			ToSection = OtherSection;
+		}
+		else
+		{
+			FromSection = OtherSection;
+			ToSection = Section;
+		}
+
+		// Check if a transition already exists for this pair
+		if (!FindTransitionForPair(FromSection, ToSection))
+		{
+			// Create a crossfade transition (default type)
+			CreateTransitionSection<UMovieSceneAnimCrossfadeTransitionSection>(FromSection, ToSection);
+		}
+	}
+}
+
+bool UMovieSceneAnimationMixerTrack::UpdateTransitionsForRow(int32 RowIndex)
+{
+	const int32 InitialSectionCount = GetAllSections().Num();
+
+	UMovieSceneAnimationMixerLayer* Layer = GetLayer(RowIndex);
+	if (!Layer)
+	{
+		return false;
+	}
+
+	// Step 1: Collect all non-transition sections on this row
+	// We collect first because subsequent operations may modify the layer's sections
+	TArray<UMovieSceneSection*> SectionsOnRow;
+	TArray<UMovieSceneAnimTransitionSectionBase*> TransitionsOnRow;
+	for (UMovieSceneSection* Sec : Layer->GetSections())
+	{
+		if (Sec)
+		{
+			if (UMovieSceneAnimTransitionSectionBase* Transition = Cast<UMovieSceneAnimTransitionSectionBase>(Sec))
+			{
+				TransitionsOnRow.Add(Transition);
+			}
+			else
+			{
+				SectionsOnRow.Add(Sec);
+			}
+		}
+	}
+
+	// Step 2: Validate existing transitions on this row - remove invalid ones
+	// Invalid cases handled by UpdateBoundsFromSourceSections:
+	// - From/To sections on different rows
+	// - No overlap between From/To
+	// - 3+ sections overlapping the transition range
+	// - One section fully contained within the other
+	for (UMovieSceneAnimTransitionSectionBase* Transition : TransitionsOnRow)
+	{
+		if (!Transition->UpdateBoundsFromSourceSections())
+		{
+			RemoveSection(*Transition);
+		}
+	}
+
+	// Step 3: Auto-create transitions for any valid 2-section overlaps
+	for (UMovieSceneSection* Sec : SectionsOnRow)
+	{
+		AutoCreateTransitionsForSection(Sec);
+	}
+
+	return GetAllSections().Num() != InitialSectionCount;
+}
+
+bool UMovieSceneAnimationMixerTrack::UpdateTransitionsForRows(const TSet<int32>& RowIndices)
+{
+	bool bAnyChanges = false;
+	for (int32 RowIndex : RowIndices)
+	{
+		bAnyChanges |= UpdateTransitionsForRow(RowIndex);
+	}
+	return bAnyChanges;
+}
+
+#if WITH_EDITOR
+EMovieSceneSectionMovedResult UMovieSceneAnimationMixerTrack::HandleTransitionSectionMoved(UMovieSceneAnimTransitionSectionBase& TransitionSection)
+{
+	if (!TransitionSection.FromSection || !TransitionSection.ToSection)
+	{
+		return EMovieSceneSectionMovedResult::None;
+	}
+
+	// 1. Handle row changes - transition must stay on same row as from/to sections
+	// This prevents vertical dragging of transitions
+	int32 CorrectRowIndex = TransitionSection.FromSection->GetRowIndex();
+	if (TransitionSection.GetRowIndex() != CorrectRowIndex)
+	{
+		TransitionSection.SetRowIndex(CorrectRowIndex);
+	}
+
+	// 2. Check if transition already matches the from/to overlap
+	// This happens when all 3 sections are moved together - no adjustment needed
+	TRange<FFrameNumber> CurrentOverlap = TRange<FFrameNumber>::Intersection(
+		TransitionSection.FromSection->GetRange(),
+		TransitionSection.ToSection->GetRange()
+	);
+	TRange<FFrameNumber> NewTransitionRange = TransitionSection.GetRange();
+
+	if (CurrentOverlap == NewTransitionRange && !CurrentOverlap.IsEmpty())
+	{
+		// Transition already matches overlap - no from/to adjustment needed
+		return EMovieSceneSectionMovedResult::None;
+	}
+
+	// 3. Handle horizontal movement - adjust from/to sections to match new transition bounds
+
+	if (NewTransitionRange.IsEmpty() || !NewTransitionRange.HasLowerBound() || !NewTransitionRange.HasUpperBound())
+	{
+		return EMovieSceneSectionMovedResult::None;
+	}
+
+	FFrameNumber TransitionStart = NewTransitionRange.GetLowerBoundValue();
+	FFrameNumber TransitionEnd = NewTransitionRange.GetUpperBoundValue();
+	FFrameNumber TransitionSize = TransitionEnd - TransitionStart;
+
+	// FROM's start is the hard limit on how far left the transition can go
+	FFrameNumber FromStart = TransitionSection.FromSection->HasStartFrame() ? TransitionSection.FromSection->GetInclusiveStartFrame() : TNumericLimits<int32>::Min();
+
+	// TO's end is the hard limit on how far right the transition can go
+	FFrameNumber ToEnd = TransitionSection.ToSection->HasEndFrame() ? TransitionSection.ToSection->GetExclusiveEndFrame() : TNumericLimits<int32>::Max();
+
+	// Check if the transition has hit its movement bounds
+	// When moving, we must preserve the transition size - don't allow shrinking via movement
+	// Users should use edge resize for that
+	FFrameNumber FinalStart = TransitionStart;
+	FFrameNumber FinalEnd = TransitionEnd;
+
+	if (TransitionStart < FromStart)
+	{
+		// Hit left bound - clamp to FROM's start, preserving size
+		FinalStart = FromStart;
+		FinalEnd = FromStart + TransitionSize;
+	}
+
+	if (FinalEnd > ToEnd)
+	{
+		// Hit right bound - clamp to TO's end, preserving size
+		FinalEnd = ToEnd;
+		FinalStart = ToEnd - TransitionSize;
+	}
+
+	// Final validation - if we still can't fit, restore original bounds
+	if (FinalStart < FromStart || FinalEnd > ToEnd || FinalStart >= FinalEnd)
+	{
+		TransitionSection.UpdateBoundsFromSourceSections();
+		return EMovieSceneSectionMovedResult::None;
+	}
+
+	// Update FROM section's end to match transition's end (creates the overlap)
+	TransitionSection.FromSection->Modify();
+	TransitionSection.FromSection->SetEndFrame(TRangeBound<FFrameNumber>::Exclusive(FinalEnd));
+
+	// Update TO section's start to match transition's start (creates the overlap)
+	TransitionSection.ToSection->Modify();
+	TransitionSection.ToSection->SetStartFrame(TRangeBound<FFrameNumber>::Inclusive(FinalStart));
+
+	// Update transition to the final range
+	TransitionSection.Modify();
+	TransitionSection.SetRange(TRange<FFrameNumber>(
+		TRangeBound<FFrameNumber>::Inclusive(FinalStart),
+		TRangeBound<FFrameNumber>::Exclusive(FinalEnd)
+	));
+
+	// If we clamped the position, the keys were already moved by MoveSection but we moved the bounds
+	// to a different position. We need to offset the keys to match the actual clamped position.
+	// Keys were moved by the original delta, but should have been moved by (FinalStart - OriginalStart).
+	// Since TransitionStart = OriginalStart + OriginalDelta, the correction is (FinalStart - TransitionStart).
+	FFrameNumber KeyCorrection = FinalStart - TransitionStart;
+	if (KeyCorrection != 0)
+	{
+		for (const FMovieSceneChannelEntry& Entry : TransitionSection.GetChannelProxy().GetAllEntries())
+		{
+			for (FMovieSceneChannel* Channel : Entry.GetChannels())
+			{
+				if (Channel)
+				{
+					Channel->Offset(KeyCorrection);
+				}
+			}
+		}
+	}
+
+	return EMovieSceneSectionMovedResult::None;
+}
+
+EMovieSceneSectionMovedResult UMovieSceneAnimationMixerTrack::OnSectionMoved(UMovieSceneSection& Section, const FMovieSceneSectionMovedParams& Params)
+{
+	// Handle transition section moves specially
+	if (UMovieSceneAnimTransitionSectionBase* TransitionSection = Cast<UMovieSceneAnimTransitionSectionBase>(&Section))
+	{
+		return HandleTransitionSectionMoved(*TransitionSection);
+	}
+
+	const bool bSectionsChanged = UpdateTransitionsForRow(Section.GetRowIndex());
+
+	return bSectionsChanged ? EMovieSceneSectionMovedResult::SectionsChanged : EMovieSceneSectionMovedResult::None;
+}
+#endif // WITH_EDITOR
+
+EMovieSceneTrackEasingSupportFlags UMovieSceneAnimationMixerTrack::SupportsEasing(FMovieSceneSupportsEasingParams& Params) const
+{
+	// Disable manual easing since the mixer track uses transition sections for blending
+	// instead of per-section easing. This prevents users from accidentally setting up
+	// easing that would conflict with transition blends.
+	return EMovieSceneTrackEasingSupportFlags::None;
+}
 
 #if WITH_EDITORONLY_DATA
 
@@ -350,6 +801,37 @@ int32 UMovieSceneAnimationMixerTrack::GetNextChildTrackRowIndex() const
 	return MaxRowIndex + 1;
 }
 
+UMovieSceneSection* UMovieSceneAnimationMixerTrack::AddNewAnimationOnRow(FFrameNumber KeyTime, UAnimSequenceBase* AnimSequence, int32 RowIndex)
+{
+	// If RowIndex is INDEX_NONE, use base class behavior which will find the next available row
+	if (RowIndex == INDEX_NONE)
+	{
+		return Super::AddNewAnimationOnRow(KeyTime, AnimSequence, RowIndex);
+	}
+
+	// For mixer tracks with explicit row index, we explicitly WANT overlaps,
+	// so we bypass InitialPlacementOnRow which would bump other sections to avoid overlap.
+	UMovieSceneSkeletalAnimationSection* NewSection = Cast<UMovieSceneSkeletalAnimationSection>(CreateNewSection());
+	if (!NewSection)
+	{
+		return nullptr;
+	}
+
+	// Calculate animation length
+	FFrameTime AnimationLength = AnimSequence->GetPlayLength() * GetTypedOuter<UMovieScene>()->GetTickResolution();
+	int32 IFrameNumber = AnimationLength.FrameNumber.Value + (int32)(AnimationLength.GetSubFrame() + 0.5f) + 1;
+
+	// Set range and row directly without bumping other sections
+	NewSection->SetRange(TRange<FFrameNumber>(KeyTime, KeyTime + IFrameNumber));
+	NewSection->SetRowIndex(RowIndex);
+	NewSection->Params.Animation = AnimSequence;
+
+	// Add via FSectionParameter to ensure proper lifecycle hooks are called
+	UMovieSceneTrack::AddSection(FSectionParameter{ *NewSection });
+
+	return NewSection;
+}
+
 void UMovieSceneAnimationMixerTrack::AddSectionWithAutoRow(UMovieSceneSection* Section)
 {
 	if (!Section)
diff --git a/Engine/Plugins/Experimental/MovieSceneAnimMixer/Source/MovieSceneAnimMixer/Private/Systems/MovieSceneAnimMixerSystem.cpp b/Engine/Plugins/Experimental/MovieSceneAnimMixer/Source/MovieSceneAnimMixer/Private/Systems/MovieSceneAnimMixerSystem.cpp
--- a/Engine/Plugins/Experimental/MovieSceneAnimMixer/Source/MovieSceneAnimMixer/Private/Systems/MovieSceneAnimMixerSystem.cpp
+++ b/Engine/Plugins/Experimental/MovieSceneAnimMixer/Source/MovieSceneAnimMixer/Private/Systems/MovieSceneAnimMixerSystem.cpp
@@ -46,6 +46,9 @@
 #include "AnimSubsystem_SequencerMixer.h"
 #include "Sections/MovieSceneSkeletalAnimationSection.h"
 #include "MirroringTask.h"
+#include "MovieSceneAnimTransitionSectionBase.h"
+#include "EvaluationVM/Tasks/TransitionEvaluationTask.h"
+#include "EntitySystem/MovieSceneInstanceRegistry.h"
 
 #include UE_INLINE_GENERATED_CPP_BY_NAME(MovieSceneAnimMixerSystem)
 
@@ -78,10 +81,12 @@ namespace UE::MovieScene
 		TMap<FObjectKey, TWeakPtr<FMovieSceneMixerRootMotionComponentData>>* RootMotions;
 		const UMovieSceneRootMotionSystem* RootMotionSystem;
 		const UMovieSceneTransformOriginSystem* TransformOriginSystem;
+		UMovieSceneEntitySystemLinker* Linker;
 
 		FEvaluateAnimMixers(TMap<FMovieSceneAnimMixerKey, TSharedPtr<FMovieSceneAnimMixer>>* InMixers, TMap<FObjectKey, TWeakPtr<FMovieSceneMixerRootMotionComponentData>>* InRootMotions, UMovieSceneEntitySystemLinker* InLinker)
 			: Mixers(InMixers)
 			, RootMotions(InRootMotions)
+			, Linker(InLinker)
 		{
 			RootMotionSystem = InLinker->FindSystem<UMovieSceneRootMotionSystem>();
 			TransformOriginSystem = InLinker->FindSystem<UMovieSceneTransformOriginSystem>();
@@ -142,9 +147,66 @@ namespace UE::MovieScene
 
 			const int32 NumMixerEntries = Mixer->WeakEntries.Num();
 
+			// ----------------------------------------------------------------------------------------------------------------------------------------
+			// Transition handling: Build transition tasks and track which sections are handled by transitions
+			// Sections handled by transitions will be skipped in the main evaluation loop below.
+			TSet<FObjectKey> SectionsInTransition;
+
+			for (int32 Index = 0; Index < NumMixerEntries; ++Index)
+			{
+				TSharedPtr<FMovieSceneAnimMixerEntry> Entry = Mixer->WeakEntries[Index].Pin();
+				if (!Entry || !Entry->bIsTransition)
+				{
+					continue;
+				}
+
+				// Get the transition section from the EntityOwner
+				UMovieSceneAnimTransitionSectionBase* TransitionSection = Cast<UMovieSceneAnimTransitionSectionBase>(Entry->EntityOwner.ResolveObjectPtr());
+				if (!TransitionSection || !TransitionSection->IsValid())
+				{
+					continue;
+				}
+
+				// Mark from/to sections as handled by this transition
+				SectionsInTransition.Add(FObjectKey(TransitionSection->FromSection));
+				SectionsInTransition.Add(FObjectKey(TransitionSection->ToSection));
+
+				// Find from/to entries to get their tasks
+				TSharedPtr<FMovieSceneAnimMixerEntry> FromEntry = Mixer->FindEntryByOwner(FObjectKey(TransitionSection->FromSection));
+				TSharedPtr<FMovieSceneAnimMixerEntry> ToEntry = Mixer->FindEntryByOwner(FObjectKey(TransitionSection->ToSection));
+
+				TSharedPtr<FAnimNextEvaluationTask> FromTask = FromEntry ? FromEntry->EvalTask : nullptr;
+				TSharedPtr<FAnimNextEvaluationTask> ToTask = ToEntry ? ToEntry->EvalTask : nullptr;
+
+				// Update the transition task with from/to tasks and blend weight.
+				// The task was created during entity import and persists across frames.
+				if (FAnimNextTransitionEvaluationTask* TransitionTask =
+					static_cast<FAnimNextTransitionEvaluationTask*>(Entry->EvalTask.Get()))
+				{
+					// Get delta time from the instance context
+					float DeltaTime = 0.0f;
+					if (Entry->InstanceHandle.IsValid())
+					{
+						const FMovieSceneContext& Context = Linker->GetInstanceRegistry()->GetContext(Entry->InstanceHandle);
+						const FFrameTime FrameDelta = Context.GetDelta();
+						const FFrameRate FrameRate = Context.GetFrameRate();
+						DeltaTime = static_cast<float>(FrameDelta.AsDecimal() / FrameRate.AsDecimal());
+					}
+
+					// The blend weight comes from the transition's BlendCurve via the WeightChannel evaluation
+					double BlendWeight = Entry->PoseWeight;
+					TransitionTask->Update(FromTask, ToTask, BlendWeight, DeltaTime);
+				}
+
+				// Reset PoseWeight to 1.0 since the transition replaces the from/to sections
+				// and should blend into the overall mix at full weight
+				Entry->PoseWeight = 1.0;
+			}
+			// ----------------------------------------------------------------------------------------------------------------------------------------
+
 			// Function that performs a lookahead accumulation of the total number of absolute weights for the specified start index's priority,
 			//  and returns a value indicating whether the entries need a separate blend stack (true) or whether they can be blended as a single operation
-			auto LookaheadComputeAccumulatedWeight = [&Mixer, NumMixerEntries](int32 StartAtIndex, int32 Priority, float& OutAccumulatedWeight) -> bool
+			auto LookaheadComputeAccumulatedWeight = [&Mixer, NumMixerEntries, &SectionsInTransition](int32 StartAtIndex, int32 Priority, float& OutAccumulatedWeight) -> bool
 			{
 				int32 NumAbsoluteBlends = 0;
 				int32 NumSkippedBlends = 0;
@@ -156,6 +218,11 @@ namespace UE::MovieScene
 					{
 						continue;
 					}
+					// Skip entries whose source sections are handled by an active transition
+					if (SectionsInTransition.Contains(NextMixerEntry->EntityOwner))
+					{
+						continue;
+					}
 					if (NextMixerEntry->Priority != Priority || NextMixerEntry->bAdditive)
 					{
 						break;
@@ -221,6 +288,12 @@ namespace UE::MovieScene
 						continue;
 					}
 
+					// Skip entries whose source sections are handled by an active transition
+					if (SectionsInTransition.Contains(MixerEntry->EntityOwner))
+					{
+						continue;
+					}
+
 					if (bNeedsSeparateWeightStack && FMath::IsNearlyZero(MixerEntry->PoseWeight, KINDA_SMALL_NUMBER))
 					{
 						continue;
@@ -701,6 +774,7 @@ void UMovieSceneAnimMixerSystem::OnRun(FSystemTaskPrerequisites& InPrerequisites
 		.ReadOptional(TrackComponents->SkeletalAnimation)
 		.ReadOptional(AnimMixerComponents->RootMotionSettings)
 		.ReadOptional(AnimMixerComponents->MirrorTable)
+		.ReadOptional(AnimMixerComponents->EntityOwner)
 		.FilterAll({ BuiltInComponents->Tags.NeedsLink})
 		.Iterate_PerEntity( &Linker->EntityManager,
 		[this, &NewMixers, BuiltInComponents, AnimMixerComponents]
@@ -714,7 +788,8 @@ void UMovieSceneAnimMixerSystem::OnRun(FSystemTaskPrerequisites& InPrerequisites
 			int32 Priority,
 			const FMovieSceneSkeletalAnimationComponentData* SkelAnimData,
 			const FMovieSceneRootMotionSettings* RootMotionSettings,
-			const TObjectPtr<UMirrorDataTable>* OptMirrorDataTable)
+			const TObjectPtr<UMirrorDataTable>* OptMirrorDataTable,
+			const FObjectKey* EntityOwner)
 			{
 				// For new entities, we 'resolve' the animation target so if 'Automatic' is picked we choose the right one automatically.
 				Target = ResolveAnimationTarget(BoundObjectKey, Target, SkelAnimData);
@@ -744,6 +819,7 @@ void UMovieSceneAnimMixerSystem::OnRun(FSystemTaskPrerequisites& InPrerequisites
 				Entry->EvalTask = Task;
 				Entry->Priority = Priority;
 				Entry->PoseWeight = 0;
+				Entry->EntityOwner = EntityOwner ? *EntityOwner : FObjectKey();
 				if (RootMotionSettings)
 				{
 					Entry->RootMotionSettings = *RootMotionSettings;
@@ -784,6 +860,7 @@ void UMovieSceneAnimMixerSystem::OnRun(FSystemTaskPrerequisites& InPrerequisites
 				Entry->bAdditive = Type.Contains(BuiltInComponents->Tags.AdditiveAnimation);
 				Entry->bRequiresBlend = Type.Contains(AnimMixerComponents->Tags.RequiresBlending);
 				Entry->MirrorTable = OptMirrorDataTable ? *OptMirrorDataTable : nullptr;
+				Entry->bIsTransition = Type.Contains(AnimMixerComponents->Tags.Transition);
 
 				TSharedPtr<FMovieSceneAnimMixer> ExistingMixer = Entry->WeakMixer.Pin();
 				if (ExistingMixer)
diff --git a/Engine/Plugins/Experimental/MovieSceneAnimMixer/Source/MovieSceneAnimMixer/Public/AnimMixerComponentTypes.h b/Engine/Plugins/Experimental/MovieSceneAnimMixer/Source/MovieSceneAnimMixer/Public/AnimMixerComponentTypes.h
--- a/Engine/Plugins/Experimental/MovieSceneAnimMixer/Source/MovieSceneAnimMixer/Public/AnimMixerComponentTypes.h
+++ b/Engine/Plugins/Experimental/MovieSceneAnimMixer/Source/MovieSceneAnimMixer/Public/AnimMixerComponentTypes.h
@@ -12,6 +12,7 @@
 #include "EvaluationVM/EvaluationTask.h"
 #include "AnimSequencerInstanceProxy.h"
 #include "Systems/MovieSceneRootMotionSystem.h"
+#include "Channels/MovieSceneInterpolation.h"
 
 enum class EMovieSceneRootMotionDestination : uint8;
 
@@ -44,9 +45,13 @@ namespace UE::MovieScene
 		TComponentTypeID<TObjectPtr<UMovieSceneAnimationMixerLayer>> MixerLayer;
 		TComponentTypeID<TObjectPtr<UMirrorDataTable>> MirrorTable;
 
+		// The object (section or track) that owns/created this entity
+		TComponentTypeID<FObjectKey> EntityOwner;
+
 		struct
 		{
 			FComponentTypeID RequiresBlending;
+			FComponentTypeID Transition;
 		} Tags;
 
 
diff --git a/Engine/Plugins/Experimental/MovieSceneAnimMixer/Source/MovieSceneAnimMixer/Public/EvaluationVM/Tasks/TransitionEvaluationTask.h b/Engine/Plugins/Experimental/MovieSceneAnimMixer/Source/MovieSceneAnimMixer/Public/EvaluationVM/Tasks/TransitionEvaluationTask.h
new file mode 100644
--- /dev/null
+++ b/Engine/Plugins/Experimental/MovieSceneAnimMixer/Source/MovieSceneAnimMixer/Public/EvaluationVM/Tasks/TransitionEvaluationTask.h
@@ -0,0 +1,59 @@
+// Copyright Epic Games, Inc. All Rights Reserved.
+
+#pragma once
+
+#include "CoreMinimal.h"
+#include "EvaluationVM/EvaluationTask.h"
+#include "TransitionEvaluationTask.generated.h"
+
+/**
+ * FAnimNextTransitionEvaluationTask
+ *
+ * Abstract base class for all Sequencer transition evaluation tasks.
+ * Provides common from/to task references and a virtual Update method
+ * that the mixer system calls each frame before execution.
+ *
+ * Derived classes must implement Execute() and can override Update() to
+ * perform additional per-frame setup, such as managing state for stateful
+ * transitions like dead blending.
+ */
+USTRUCT()
+struct MOVIESCENEANIMMIXER_API FAnimNextTransitionEvaluationTask : public FAnimNextEvaluationTask
+{
+	GENERATED_BODY()
+
+	DECLARE_ANIM_EVALUATION_TASK(FAnimNextTransitionEvaluationTask)
+
+	FAnimNextTransitionEvaluationTask() = default;
+	virtual ~FAnimNextTransitionEvaluationTask() = default;
+
+	/**
+	 * Update the task with the current frame's data.
+	 * Called by the mixer system each frame before the task is executed.
+	 *
+	 * The base implementation stores the from/to tasks and blend weight.
+	 * Derived classes should call the base implementation and then perform
+	 * any additional per-frame setup.
+	 */
+	virtual void Update(
+		const TSharedPtr<FAnimNextEvaluationTask>& InFromTask,
+		const TSharedPtr<FAnimNextEvaluationTask>& InToTask,
+		double InBlendWeight,
+		float InDeltaTime);
+
+	virtual void Execute(UE::UAF::FEvaluationVM& VM) const override { LowLevelFatalError(TEXT("Pure virtual not implemented (FAnimNextTransitionEvaluationTask::Execute)")); }
+
+protected:
+
+	/** The task that evaluates the "from" section's pose */
+	TSharedPtr<FAnimNextEvaluationTask> FromTask;
+
+	/** The task that evaluates the "to" section's pose */
+	TSharedPtr<FAnimNextEvaluationTask> ToTask;
+
+	/** Current blend weight (0 = full from, 1 = full to) */
+	double BlendWeight = 0.0;
+
+	/** Delta time since last update */
+	float DeltaTime = 0.0f;
+};
diff --git a/Engine/Plugins/Experimental/MovieSceneAnimMixer/Source/MovieSceneAnimMixer/Public/MovieSceneAnimCrossfadeTransitionSection.h b/Engine/Plugins/Experimental/MovieSceneAnimMixer/Source/MovieSceneAnimMixer/Public/MovieSceneAnimCrossfadeTransitionSection.h
new file mode 100644
--- /dev/null
+++ b/Engine/Plugins/Experimental/MovieSceneAnimMixer/Source/MovieSceneAnimMixer/Public/MovieSceneAnimCrossfadeTransitionSection.h
@@ -0,0 +1,44 @@
+// Copyright Epic Games, Inc. All Rights Reserved.
+
+#pragma once
+
+#include "CoreTypes.h"
+#include "Channels/MovieSceneFloatChannel.h"
+#include "MovieSceneAnimTransitionSectionBase.h"
+#include "MovieSceneAnimCrossfadeTransitionSection.generated.h"
+
+/**
+ * Crossfade transition section that blends between two animation sections using a blend curve.
+ * The blend curve controls the interpolation: 0.0 = full FromSection, 1.0 = full ToSection.
+ */
+UCLASS(MinimalAPI, DisplayName="Crossfade Transition")
+class UMovieSceneAnimCrossfadeTransitionSection : public UMovieSceneAnimTransitionSectionBase
+{
+	GENERATED_BODY()
+
+public:
+
+	UMovieSceneAnimCrossfadeTransitionSection(const FObjectInitializer& ObjInit);
+
+	/**
+	 * The blend curve controlling the transition.
+	 * Value of 0.0 = full FromSection pose
+	 * Value of 1.0 = full ToSection pose
+	 */
+	UPROPERTY()
+	FMovieSceneFloatChannel BlendCurve;
+
+	/** Initialize the blend curve with a default linear ramp from 0 to 1 */
+	MOVIESCENEANIMMIXER_API virtual void InitializeDefaultCurve() override;
+
+	//~ UMovieSceneAnimTransitionSectionBase interface
+	virtual TSharedPtr<FAnimNextTransitionEvaluationTask> CreateTransitionTask() const override;
+	virtual FName GetTransitionIconStyleName() const override;
+	virtual FText GetTransitionDisplayName() const override;
+
+protected:
+
+	//~ UMovieSceneAnimTransitionSectionBase interface
+	virtual void RebuildChannelProxy(FMovieSceneChannelProxyData& Channels) override;
+	virtual const FMovieSceneFloatChannel* GetBlendWeightChannel() const override { return &BlendCurve; }
+};
diff --git a/Engine/Plugins/Experimental/MovieSceneAnimMixer/Source/MovieSceneAnimMixer/Public/MovieSceneAnimInertialDeadBlendTransitionSection.h b/Engine/Plugins/Experimental/MovieSceneAnimMixer/Source/MovieSceneAnimMixer/Public/MovieSceneAnimInertialDeadBlendTransitionSection.h
new file mode 100644
--- /dev/null
+++ b/Engine/Plugins/Experimental/MovieSceneAnimMixer/Source/MovieSceneAnimMixer/Public/MovieSceneAnimInertialDeadBlendTransitionSection.h
@@ -0,0 +1,94 @@
+// Copyright Epic Games, Inc. All Rights Reserved.
+
+#pragma once
+
+#include "CoreMinimal.h"
+#include "AlphaBlend.h"
+#include "MovieSceneAnimTransitionSectionBase.h"
+#include "MovieSceneAnimInertialDeadBlendTransitionSection.generated.h"
+
+class UCurveFloat;
+
+/**
+ * Inertial dead blend transition section that provides smooth, velocity-aware blending
+ * between two animation sections.
+ * 
+ * Unlike crossfade which simply interpolates between poses, dead blending:
+ * - Captures the velocity of the outgoing animation
+ * - Extrapolates that motion forward during the blend
+ * - Provides smoother transitions especially for fast-moving animations
+ */
+UCLASS(MinimalAPI, DisplayName="Inertial Dead Blend Transition")
+class UMovieSceneAnimInertialDeadBlendTransitionSection : public UMovieSceneAnimTransitionSectionBase
+{
+	GENERATED_BODY()
+
+public:
+
+	UMovieSceneAnimInertialDeadBlendTransitionSection(const FObjectInitializer& ObjInit);
+
+	// ---- Blend Parameters ----
+
+	/** The blend mode to use for the alpha curve */
+	UPROPERTY(EditAnywhere, Category = "Blending")
+	EAlphaBlendOption BlendMode = EAlphaBlendOption::HermiteCubic;
+
+	/** Optional custom blend curve (used when BlendMode is set to Custom) */
+	UPROPERTY(EditAnywhere, Category = "Blending")
+	TObjectPtr<UCurveFloat> CustomBlendCurve = nullptr;
+
+	// ---- Extrapolation Parameters ----
+
+	/**
+	 * The average half-life of decay in seconds to use when extrapolating the animation.
+	 * This value is scaled based on how much the outgoing velocity aligns with the
+	 * direction toward the target pose.
+	 */
+	UPROPERTY(EditAnywhere, Category = "Extrapolation", meta = (ClampMin = "0.0", UIMin = "0.0"))
+	float ExtrapolationHalfLife = 1.0f;
+
+	/**
+	 * The minimum half-life of decay in seconds. Used when velocities are very small
+	 * or moving away from the target pose.
+	 */
+	UPROPERTY(EditAnywhere, Category = "Extrapolation", meta = (ClampMin = "0.0", UIMin = "0.0"))
+	float ExtrapolationHalfLifeMin = 0.05f;
+
+	/**
+	 * The maximum half-life of decay in seconds. Limits how long extrapolation can
+	 * continue when velocities are small but aligned with the target.
+	 */
+	UPROPERTY(EditAnywhere, Category = "Extrapolation", meta = (ClampMin = "0.0", UIMin = "0.0"))
+	float ExtrapolationHalfLifeMax = 1.0f;
+
+	/**
+	 * Maximum translation velocity in cm/s. Velocities above this are clamped
+	 * to prevent extreme extrapolation.
+	 */
+	UPROPERTY(EditAnywhere, Category = "Extrapolation", meta = (ClampMin = "0.0", UIMin = "0.0"))
+	float MaximumTranslationVelocity = 500.0f;
+
+	/**
+	 * Maximum rotation velocity in degrees/s. Velocities above this are clamped
+	 * to prevent extreme extrapolation.
+	 */
+	UPROPERTY(EditAnywhere, Category = "Extrapolation", meta = (ClampMin = "0.0", UIMin = "0.0"))
+	float MaximumRotationVelocity = 360.0f;
+
+	/**
+	 * Maximum scale velocity. Velocities above this are clamped
+	 * to prevent extreme extrapolation.
+	 */
+	UPROPERTY(EditAnywhere, Category = "Extrapolation", meta = (ClampMin = "0.0", UIMin = "0.0"))
+	float MaximumScaleVelocity = 4.0f;
+
+	//~ UMovieSceneAnimTransitionSectionBase interface
+	virtual TSharedPtr<FAnimNextTransitionEvaluationTask> CreateTransitionTask() const override;
+	virtual FName GetTransitionIconStyleName() const override;
+	virtual FText GetTransitionDisplayName() const override;
+
+protected:
+
+	//~ UMovieSceneAnimTransitionSectionBase interface
+	virtual void RebuildChannelProxy(FMovieSceneChannelProxyData& Channels) override;
+};
diff --git a/Engine/Plugins/Experimental/MovieSceneAnimMixer/Source/MovieSceneAnimMixer/Public/MovieSceneAnimTransitionSectionBase.h b/Engine/Plugins/Experimental/MovieSceneAnimMixer/Source/MovieSceneAnimMixer/Public/MovieSceneAnimTransitionSectionBase.h
new file mode 100644
--- /dev/null
+++ b/Engine/Plugins/Experimental/MovieSceneAnimMixer/Source/MovieSceneAnimMixer/Public/MovieSceneAnimTransitionSectionBase.h
@@ -0,0 +1,109 @@
+// Copyright Epic Games, Inc. All Rights Reserved.
+
+#pragma once
+
+#include "CoreTypes.h"
+#include "Channels/MovieSceneFloatChannel.h"
+#include "EntitySystem/IMovieSceneEntityProvider.h"
+#include "Internationalization/Text.h"
+#include "MovieSceneSection.h"
+#include "MovieSceneAnimationMixerTrack.h"
+#include "MovieSceneAnimTransitionSectionBase.generated.h"
+
+class UMovieSceneAnimationMixerTrack;
+class UMovieSceneAnimationMixerLayer;
+struct FAnimNextEvaluationTask;
+
+struct FAnimNextTransitionEvaluationTask;
+
+/**
+ * Abstract base class for all Anim Mixer transition section types.
+ * Transition sections handle blending between two overlapping animation sections on the same layer.
+ *
+ * Transition tasks follow a create-once, update-each-frame pattern:
+ * 1. CreateTransitionTask() is called once during entity import to create the task
+ * 2. The mixer system calls Update() on the task each frame to set from/to tasks and timing
+ * 3. The task persists as an entity component, allowing it to maintain state across frames
+ */
+UCLASS(Abstract, MinimalAPI)
+class UMovieSceneAnimTransitionSectionBase
+	: public UMovieSceneSection
+	, public IMovieSceneEntityProvider
+	, public IMovieSceneAnimationMixerItemInterface
+{
+	GENERATED_BODY()
+
+public:
+
+	UMovieSceneAnimTransitionSectionBase(const FObjectInitializer& ObjInit);
+
+	/** The section that this transition blends FROM (earlier in time) */
+	UPROPERTY()
+	TObjectPtr<UMovieSceneSection> FromSection;
+
+	/** The section that this transition blends TO (later in time) */
+	UPROPERTY()
+	TObjectPtr<UMovieSceneSection> ToSection;
+
+	/** Check if this transition is structurally valid (sections exist, same row, bounded, overlap exists, neither fully contained) */
+	MOVIESCENEANIMMIXER_API bool IsValid() const;
+
+	/** Update this transition's bounds to match the current overlap of from/to sections */
+	MOVIESCENEANIMMIXER_API bool UpdateBoundsFromSourceSections();
+
+	/** Scale blend weight keys proportionally from old range to new range */
+	MOVIESCENEANIMMIXER_API void ScaleKeysToRange(const TRange<FFrameNumber>& OldRange, const TRange<FFrameNumber>& NewRange);
+
+	/** Get the computed overlap range between from and to sections */
+	MOVIESCENEANIMMIXER_API TRange<FFrameNumber> ComputeOverlapRange() const;
+
+	/** Get transition progress [0,1] at the given time. 0 = start of transition, 1 = end */
+	MOVIESCENEANIMMIXER_API float GetTransitionProgress(FFrameTime InTime) const;
+
+	/** Called after the section range is set to initialize any default curves/channels */
+	virtual void InitializeDefaultCurve() {}
+
+	/**
+	 * Create the evaluation task for this transition.
+	 * Called once during entity import. The task will persist as an entity component
+	 * and have its Update() method called each frame by the mixer system.
+	 */
+	virtual TSharedPtr<FAnimNextTransitionEvaluationTask> CreateTransitionTask() const
+		PURE_VIRTUAL(UMovieSceneAnimTransitionSectionBase::CreateTransitionTask, return nullptr;);
+
+	/**
+	 * Get the blend weight channel for evaluation.
+	 * The result will be used as the blend weight (0 = from, 1 = to).
+	 */
+	virtual const FMovieSceneFloatChannel* GetBlendWeightChannel() const { return nullptr; }
+
+	/**
+	 * Get the icon style name for this transition type.
+	 * Used by the editor to display a small icon in the transition section.
+	 */
+	virtual FName GetTransitionIconStyleName() const { return NAME_None; }
+
+	/**
+	 * Get a short display name for this transition type.
+	 * Used for tooltips and other UI elements.
+	 */
+	virtual FText GetTransitionDisplayName() const { return FText::GetEmpty(); }
+
+	//~ IMovieSceneAnimationMixerItemInterface
+	virtual FColor GetMixerItemTint() const override { return TransitionTintColor; }
+
+	//~ UMovieSceneSection interface
+	virtual EMovieSceneChannelProxyType CacheChannelProxy() override final;
+
+	//~ IMovieSceneEntityProvider interface
+	virtual void ImportEntityImpl(UMovieSceneEntitySystemLinker* EntityLinker, const FEntityImportParams& Params, FImportedEntity* OutImportedEntity) override;
+	virtual bool PopulateEvaluationFieldImpl(const TRange<FFrameNumber>& EffectiveRange, const FMovieSceneEvaluationFieldEntityMetaData& InMetaData, FMovieSceneEntityComponentFieldBuilder* OutFieldBuilder) override;
+
+protected:
+
+	/** Subclasses implement to register their specific channels */
+	virtual void RebuildChannelProxy(FMovieSceneChannelProxyData& Channels) PURE_VIRTUAL(UMovieSceneAnimTransitionSectionBase::RebuildChannelProxy, );
+
+	/** Tint color for transition sections in the timeline */
+	FColor TransitionTintColor = FColor(120, 80, 160, 200); 
+};
diff --git a/Engine/Plugins/Experimental/MovieSceneAnimMixer/Source/MovieSceneAnimMixer/Public/MovieSceneAnimationMixerTrack.h b/Engine/Plugins/Experimental/MovieSceneAnimMixer/Source/MovieSceneAnimMixer/Public/MovieSceneAnimationMixerTrack.h
--- a/Engine/Plugins/Experimental/MovieSceneAnimMixer/Source/MovieSceneAnimMixer/Public/MovieSceneAnimationMixerTrack.h
+++ b/Engine/Plugins/Experimental/MovieSceneAnimMixer/Source/MovieSceneAnimMixer/Public/MovieSceneAnimationMixerTrack.h
@@ -14,9 +14,12 @@
 #include "StructUtils/InstancedStruct.h"
 #include "Tracks/MovieSceneCommonAnimationTrack.h"
 #include "Systems/MovieSceneRootMotionSystem.h"
+#include "MovieSceneTrack.h"
 #include "MovieSceneAnimationMixerTrack.generated.h"
 
 class UMovieSceneAnimationMixerLayer;
+class UMovieSceneAnimTransitionSectionBase;
+class UMovieSceneAnimCrossfadeTransitionSection;
 
 USTRUCT()
 struct FMovieSceneByteChannelDefaultOnly : public FMovieSceneByteChannel
@@ -120,14 +123,21 @@ public:
 	virtual bool FixRowIndices() override;
 	virtual bool SupportsType(TSubclassOf<UMovieSceneSection> SectionClass) const override;
 
+	//~ IMovieSceneTrackVirtualAPI interface
+	virtual void RemoveSection(UMovieSceneSection& Section) override;
+
 	virtual void OnSectionAddedImpl(UMovieSceneSection* Secton) override;
 	virtual void OnSectionRemovedImpl(UMovieSceneSection* Section) override;
+#if WITH_EDITOR
+	virtual EMovieSceneSectionMovedResult OnSectionMoved(UMovieSceneSection& Section, const FMovieSceneSectionMovedParams& Params) override;
+#endif
+	virtual EMovieSceneTrackEasingSupportFlags SupportsEasing(FMovieSceneSupportsEasingParams& Params) const override;
 
+	MOVIESCENEANIMMIXER_API virtual UMovieSceneSection* AddNewAnimationOnRow(FFrameNumber KeyTime, UAnimSequenceBase* AnimSequence, int32 RowIndex) override;
 	MOVIESCENEANIMMIXER_API virtual void AddSectionWithAutoRow(UMovieSceneSection* Section) override;
 	MOVIESCENEANIMMIXER_API virtual int32 ComputeNextAvailableRowIndex() const override;
 
 #if WITH_EDITORONLY_DATA
-
 	virtual FText GetTrackRowDisplayName(int32 RowIndex) const override;
 	virtual FText GetDefaultDisplayName() const override;
 	virtual bool CanRename() const override;
@@ -167,11 +177,60 @@ public:
 	/** Insert a new empty layer at the specified index, shifting existing layers down */
 	MOVIESCENEANIMMIXER_API UMovieSceneAnimationMixerLayer* InsertLayer(int32 InsertIndex);
 
-public: 
+	// Transition Section Management
+
+	/**
+	 * Create a transition section between two overlapping sections on the same layer.
+	 */
+	template<typename TransitionClass>
+	TransitionClass* CreateTransitionSection(UMovieSceneSection* FromSection, UMovieSceneSection* ToSection);
+
+	/**
+	 * Create a transition section of the specified type between two overlapping sections (runtime class selection).
+	 */
+	MOVIESCENEANIMMIXER_API UMovieSceneAnimTransitionSectionBase* CreateTransitionSectionOfType(
+		UMovieSceneSection* FromSection,
+		UMovieSceneSection* ToSection,
+		TSubclassOf<UMovieSceneAnimTransitionSectionBase> TransitionClass);
+
+	/** Find transition sections that reference the given section as either from or to */
+	MOVIESCENEANIMMIXER_API TArray<UMovieSceneAnimTransitionSectionBase*> FindTransitionsForSection(UMovieSceneSection* Section) const;
+
+	/** Find sections that overlap with the given section on the same layer (excluding transitions) */
+	MOVIESCENEANIMMIXER_API TArray<UMovieSceneSection*> FindOverlappingSections(UMovieSceneSection* Section) const;
+
+	/** Count how many non-transition sections overlap the given range on the specified row */
+	MOVIESCENEANIMMIXER_API int32 CountSectionsOverlappingRange(const TRange<FFrameNumber>& Range, int32 RowIndex) const;
+
+	/** Check if a transition already exists for the given from/to section pair */
+	MOVIESCENEANIMMIXER_API UMovieSceneAnimTransitionSectionBase* FindTransitionForPair(UMovieSceneSection* FromSection, UMovieSceneSection* ToSection) const;
+
+	/** Auto-create transitions for any overlapping sections involving the given section */
+	MOVIESCENEANIMMIXER_API void AutoCreateTransitionsForSection(UMovieSceneSection* Section);
+
+	/**
+	 * Validates and updates all transitions on the specified row.
+	 * This removes invalid transitions and creates new ones for valid overlaps.
+	 * Call this after any operation that might affect transition validity on a row
+	 * (section added, moved, removed, or row changed).
+	 */
+	MOVIESCENEANIMMIXER_API bool UpdateTransitionsForRow(int32 RowIndex);
+
+	/**
+	 * Validates and updates all transitions on the specified rows.
+	 * Convenience method that calls UpdateTransitionsForRow for each row.
+	 */
+	MOVIESCENEANIMMIXER_API bool UpdateTransitionsForRows(const TSet<int32>& RowIndices);
+
+public:
 	UPROPERTY(EditAnywhere, Category = "Mixed Animation")
 	TInstancedStruct<FMovieSceneMixedAnimationTarget> MixedAnimationTarget;
 
 private:
+#if WITH_EDITOR
+	/** Handle when a transition section is moved - adjusts from/to sections and constrains row changes */
+	EMovieSceneSectionMovedResult HandleTransitionSectionMoved(UMovieSceneAnimTransitionSectionBase& TransitionSection);
+#endif
 	// Row assignment for child tracks
 	UPROPERTY()
 	TMap<TObjectPtr<UMovieSceneTrack>, int32> ChildTracks;
@@ -181,3 +240,52 @@ private:
 	UPROPERTY()
 	TArray<TObjectPtr<UMovieSceneAnimationMixerLayer>> Layers;
 };
+
+template<typename TransitionClass>
+TransitionClass* UMovieSceneAnimationMixerTrack::CreateTransitionSection(UMovieSceneSection* FromSection, UMovieSceneSection* ToSection)
+{
+	static_assert(TIsDerivedFrom<TransitionClass, UMovieSceneAnimTransitionSectionBase>::Value,
+		"TransitionClass must derive from UMovieSceneAnimTransitionSectionBase");
+
+	if (!FromSection || !ToSection)
+	{
+		return nullptr;
+	}
+
+	// Sections must be on the same row (layer)
+	if (FromSection->GetRowIndex() != ToSection->GetRowIndex())
+	{
+		return nullptr;
+	}
+
+	// Compute the overlap range
+	TRange<FFrameNumber> Overlap = TRange<FFrameNumber>::Intersection(
+		FromSection->GetRange(),
+		ToSection->GetRange()
+	);
+
+	if (Overlap.IsEmpty())
+	{
+		return nullptr;
+	}
+
+	Modify();
+
+	// Create the transition section
+	TransitionClass* TransitionSection = NewObject<TransitionClass>(this, TransitionClass::StaticClass(), NAME_None, RF_Transactional);
+	TransitionSection->SetRange(Overlap);
+	TransitionSection->SetRowIndex(FromSection->GetRowIndex());
+	TransitionSection->FromSection = FromSection;
+	TransitionSection->ToSection = ToSection;
+
+	// Ensure transition sections always draw on top of regular sections
+	TransitionSection->SetOverlapPriority(INT32_MAX);
+
+	// Initialize any default curves now that the range is set
+	TransitionSection->InitializeDefaultCurve();
+
+	// Add via FSectionParameter to ensure proper lifecycle hooks are called
+	UMovieSceneTrack::AddSection(FSectionParameter{*TransitionSection});
+
+	return TransitionSection;
+}
diff --git a/Engine/Plugins/Experimental/MovieSceneAnimMixer/Source/MovieSceneAnimMixer/Public/Systems/MovieSceneAnimMixerSystem.h b/Engine/Plugins/Experimental/MovieSceneAnimMixer/Source/MovieSceneAnimMixer/Public/Systems/MovieSceneAnimMixerSystem.h
--- a/Engine/Plugins/Experimental/MovieSceneAnimMixer/Source/MovieSceneAnimMixer/Public/Systems/MovieSceneAnimMixerSystem.h
+++ b/Engine/Plugins/Experimental/MovieSceneAnimMixer/Source/MovieSceneAnimMixer/Public/Systems/MovieSceneAnimMixerSystem.h
@@ -25,7 +25,7 @@ struct FMovieSceneAnimMixer;
 
 struct FMovieSceneAnimMixerEntry
 {
-	
+
 	TWeakPtr<FMovieSceneAnimMixer> WeakMixer;
 	TSharedPtr<FAnimNextEvaluationTask> EvalTask;
 	/** Shared pointer to the root motion for this entry if it came from a FMovieSceneRootMotionSettings on a mixer
@@ -40,6 +40,10 @@ struct FMovieSceneAnimMixerEntry
 	TObjectPtr<UMirrorDataTable> MirrorTable;
 	UE::MovieScene::FMovieSceneEntityID EntityID;
 	UE::MovieScene::FInstanceHandle InstanceHandle;
+	/** The section or track that owns this entity - used for transition section lookups */
+	FObjectKey EntityOwner;
+	/** True if this entry is from a transition section */
+	bool bIsTransition = false;
 
 
 	inline bool operator<(const FMovieSceneAnimMixerEntry& RHS) const
@@ -87,6 +91,25 @@ struct FMovieSceneAnimMixer
 	TWeakPtr<FMovieSceneMixerRootMotionComponentData> WeakRootMotion;
 
 	bool bNeedsResort = false;
+
+	/**
+	 * Find a mixer entry by its EntityOwner (the section or track that created it).
+	 * Used by transition sections to find from/to entries for building transition programs.
+	 */
+	TSharedPtr<FMovieSceneAnimMixerEntry> FindEntryByOwner(const FObjectKey& InOwner) const
+	{
+		for (const TWeakPtr<FMovieSceneAnimMixerEntry>& WeakEntry : WeakEntries)
+		{
+			if (TSharedPtr<FMovieSceneAnimMixerEntry> Entry = WeakEntry.Pin())
+			{
+				if (Entry->EntityOwner == InOwner)
+				{
+					return Entry;
+				}
+			}
+		}
+		return nullptr;
+	}
 };
 
 
diff --git a/Engine/Plugins/Experimental/MovieSceneAnimMixer/Source/MovieSceneAnimMixerEditor/MovieSceneAnimMixerEditor.Build.cs b/Engine/Plugins/Experimental/MovieSceneAnimMixer/Source/MovieSceneAnimMixerEditor/MovieSceneAnimMixerEditor.Build.cs
--- a/Engine/Plugins/Experimental/MovieSceneAnimMixer/Source/MovieSceneAnimMixerEditor/MovieSceneAnimMixerEditor.Build.cs
+++ b/Engine/Plugins/Experimental/MovieSceneAnimMixer/Source/MovieSceneAnimMixerEditor/MovieSceneAnimMixerEditor.Build.cs
@@ -29,6 +29,7 @@ namespace UnrealBuildTool.Rules
 					"MovieSceneTools",
 					"Sequencer",
 					"SequencerCore",
+					"SequencerWidgets",
 					"Slate",
 					"SlateCore",
 					"UAF",
diff --git a/Engine/Plugins/Experimental/MovieSceneAnimMixer/Source/MovieSceneAnimMixerEditor/Private/AnimTransitionSectionInterface.cpp b/Engine/Plugins/Experimental/MovieSceneAnimMixer/Source/MovieSceneAnimMixerEditor/Private/AnimTransitionSectionInterface.cpp
new file mode 100644
--- /dev/null
+++ b/Engine/Plugins/Experimental/MovieSceneAnimMixer/Source/MovieSceneAnimMixerEditor/Private/AnimTransitionSectionInterface.cpp
@@ -0,0 +1,465 @@
+// Copyright Epic Games, Inc. All Rights Reserved.
+
+#include "AnimTransitionSectionInterface.h"
+
+#include "MovieSceneAnimTransitionSectionBase.h"
+#include "UObject/UObjectIterator.h"
+#include "MovieSceneAnimCrossfadeTransitionSection.h"
+#include "MovieSceneAnimationMixerTrack.h"
+#include "ISectionLayoutBuilder.h"
+#include "SequencerSectionPainter.h"
+#include "MVVM/ViewModels/SectionModel.h"
+#include "MVVM/ViewModels/SectionOutlinerModel.h"
+#include "Fonts/FontMeasure.h"
+#include "Framework/Application/SlateApplication.h"
+#include "Framework/MultiBox/MultiBoxBuilder.h"
+#include "ScopedTransaction.h"
+#include "Styling/AppStyle.h"
+#include "Rendering/SlateRenderer.h"
+#include "TimeToPixel.h"
+
+#define LOCTEXT_NAMESPACE "AnimTransitionSectionInterface"
+
+namespace UE::Sequencer
+{
+
+FAnimTransitionSectionInterface::FAnimTransitionSectionInterface(UMovieSceneSection& InSection, UMovieSceneTrack& InTrack, FGuid InObjectBinding)
+	: SectionPtr(&InSection)
+	, TrackPtr(&InTrack)
+	, ObjectBindingId(InObjectBinding)
+{
+}
+
+UMovieSceneSection* FAnimTransitionSectionInterface::GetSectionObject()
+{
+	return SectionPtr.Get();
+}
+
+FText FAnimTransitionSectionInterface::GetSectionTitle() const
+{
+	// Empty for now- curves and icons show the transition- title just gets in the way
+	return FText::GetEmpty();
+}
+
+FText FAnimTransitionSectionInterface::GetSectionToolTip() const
+{
+	if (UMovieSceneAnimTransitionSectionBase* TransitionSection = GetTransitionSection())
+	{
+		return TransitionSection->GetTransitionDisplayName();
+	}
+	return LOCTEXT("TransitionTooltip", "Transition");
+}
+
+UMovieSceneAnimTransitionSectionBase* FAnimTransitionSectionInterface::GetTransitionSection() const
+{
+	return Cast<UMovieSceneAnimTransitionSectionBase>(SectionPtr.Get());
+}
+
+int32 FAnimTransitionSectionInterface::OnPaintSection(FSequencerSectionPainter& Painter) const
+{
+	int32 LayerId = Painter.LayerId;
+
+	UMovieSceneAnimTransitionSectionBase* TransitionSection = GetTransitionSection();
+	if (!TransitionSection)
+	{
+		return LayerId;
+	}
+
+	const TRange<FFrameNumber> SectionRange = TransitionSection->GetRange();
+	if (SectionRange.IsEmpty() || !SectionRange.HasLowerBound() || !SectionRange.HasUpperBound())
+	{
+		return LayerId;
+	}
+
+	// Check if we're in "section outliner mode" (layer expanded, each section on its own row)
+	// vs "overlay mode" (layer collapsed, sections overlapping on the same row).
+	// In outliner mode, we paint a solid background since sections don't overlap.
+	// In overlay mode, we use transparent fill to let underlying sections show through.
+	const bool bIsInOutlinerMode = Painter.SectionModel && Painter.SectionModel->FindAncestorOfType<FSectionOutlinerModel>().IsValid();
+
+	if (bIsInOutlinerMode)
+	{
+		// Paint teal background to distinguish transitions from regular animation sections
+		// Teal suggests "connection/linking" and feels active (not disabled like gray would)
+		const FLinearColor TealBackground(0.08f, 0.22f, 0.25f, 1.0f);
+		LayerId = Painter.PaintSectionBackground(TealBackground);
+	}
+
+	// Use HeaderGeometry to only draw on the section header row, not expanded key areas
+	const FVector2D HeaderSize = Painter.HeaderGeometry.GetLocalSize();
+	const float CurveHeight = HeaderSize.Y - 2.0f;  // 1px padding top and bottom
+
+	// Get the blend weight channel to sample the curve
+	const FMovieSceneFloatChannel* BlendChannel = TransitionSection->GetBlendWeightChannel();
+
+	// Get the time converter for pixel calculations
+	const FTimeToPixel& TimeToPixel = Painter.GetTimeConverter();
+
+	// Get the brush for the filled curve area
+	const FSlateBrush* CurveBrush = FAppStyle::Get().GetBrush("Sequencer.Timeline.EaseInOut");
+	FSlateResourceHandle ResourceHandle = FSlateApplication::Get().GetRenderer()->GetResourceHandle(*CurveBrush);
+	const FSlateShaderResourceProxy* ResourceProxy = ResourceHandle.GetResourceProxy();
+	FVector2f AtlasOffset = ResourceProxy ? ResourceProxy->StartUV : FVector2f(0.f, 0.f);
+	FVector2f AtlasUVSize = ResourceProxy ? ResourceProxy->SizeUV : FVector2f(1.f, 1.f);
+
+	// Prepare vertices and indices for the filled curve
+	TArray<FSlateVertex> Verts;
+	TArray<SlateIndex> Indices;
+	TArray<FVector2D> BorderPoints;
+
+	// Sample the blend curve at regular intervals
+	const int32 NumSamples = FMath::Max(2, FMath::CeilToInt(HeaderSize.X / 4.0f));  // Sample every ~4 pixels
+	Verts.Reserve(NumSamples * 2);
+	BorderPoints.Reserve(NumSamples);
+
+	const FFrameNumber StartFrame = SectionRange.GetLowerBoundValue();
+	const FFrameNumber EndFrame = SectionRange.GetUpperBoundValue();
+	const double FrameRange = (EndFrame - StartFrame).Value;
+
+	const FVector2f HeaderSizeF(HeaderSize);
+	const FVector2f Pos(Painter.HeaderGeometry.GetAbsolutePosition());
+	const FSlateRenderTransform RenderTransform;
+
+	// In overlay mode, use semi-transparent fill; in outliner mode, use more opaque fill
+	const FColor FillColor = bIsInOutlinerMode ? FColor(0, 0, 0, 100) : FColor(0, 0, 0, 51);
+
+	for (int32 i = 0; i < NumSamples; ++i)
+	{
+		// Calculate the normalized position (0 to 1) across the transition
+		const float Alpha = static_cast<float>(i) / static_cast<float>(NumSamples - 1);
+
+		// Calculate the frame time at this sample point
+		const FFrameTime SampleTime = FFrameTime(StartFrame) + FFrameTime::FromDecimal(Alpha * FrameRange);
+
+		// Get the blend weight at this time (default to linear if no channel)
+		float BlendWeight = Alpha;  // Default linear
+		if (BlendChannel)
+		{
+			float ChannelValue = 0.0f;
+			BlendChannel->Evaluate(SampleTime, ChannelValue);
+			BlendWeight = FMath::Clamp(ChannelValue, 0.0f, 1.0f);
+		}
+
+		// Calculate the X position in local space
+		const float X = Alpha * HeaderSize.X;
+
+		// Calculate Y position: BlendWeight of 0 = bottom, BlendWeight of 1 = top
+		// This shows the "to" section growing from bottom as blend increases
+		const float CurveY = 1.0f + (1.0f - BlendWeight) * CurveHeight;
+
+		// UV coordinates for the texture
+		const FVector2f UV_Top(Alpha, 0.0f);
+		const FVector2f UV_Curve(Alpha, 0.5f);
+
+		// Add top vertex (at y=1, just below top edge)
+		Verts.Add(FSlateVertex::Make<ESlateVertexRounding::Disabled>(
+			RenderTransform,
+			Pos + FVector2f(X, 1.0f) * Painter.HeaderGeometry.Scale,
+			AtlasOffset + UV_Top * AtlasUVSize,
+			FillColor));
+
+		// Add curve vertex
+		Verts.Add(FSlateVertex::Make<ESlateVertexRounding::Disabled>(
+			RenderTransform,
+			Pos + FVector2f(X, CurveY) * Painter.HeaderGeometry.Scale,
+			AtlasOffset + UV_Curve * AtlasUVSize,
+			FillColor));
+
+		// Store border point for the curve line
+		BorderPoints.Add(FVector2D(X, CurveY));
+
+		// Add indices for triangles (after we have at least 4 vertices)
+		if (Verts.Num() >= 4)
+		{
+			const int32 Index0 = Verts.Num() - 4;  // Previous top
+			const int32 Index1 = Verts.Num() - 3;  // Previous curve
+			const int32 Index2 = Verts.Num() - 2;  // Current top
+			const int32 Index3 = Verts.Num() - 1;  // Current curve
+
+			// Two triangles to form a quad
+			Indices.Add(Index0);
+			Indices.Add(Index1);
+			Indices.Add(Index2);
+
+			Indices.Add(Index1);
+			Indices.Add(Index2);
+			Indices.Add(Index3);
+		}
+	}
+
+	// Draw the filled curve area
+	if (Indices.Num() > 0)
+	{
+		FSlateDrawElement::MakeCustomVerts(
+			Painter.DrawElements,
+			LayerId,
+			ResourceHandle,
+			Verts,
+			Indices,
+			nullptr,
+			0,
+			0,
+			ESlateDrawEffect::PreMultipliedAlpha);
+	}
+	++LayerId;
+
+	// Draw the curve border line
+	if (BorderPoints.Num() >= 2)
+	{
+		const FLinearColor BorderColor(0.4f, 0.7f, 0.75f, 0.9f);  // Light teal/cyan border
+		FSlateDrawElement::MakeLines(
+			Painter.DrawElements,
+			LayerId,
+			Painter.HeaderGeometry.ToPaintGeometry(),
+			BorderPoints,
+			ESlateDrawEffect::None,
+			BorderColor,
+			true,  // bAntialias
+			1.5f);
+	}
+	++LayerId;
+
+	// Draw the transition type icon in the top-left corner
+	const FName IconStyleName = TransitionSection->GetTransitionIconStyleName();
+	if (!IconStyleName.IsNone())
+	{
+		const FSlateBrush* IconBrush = FAppStyle::Get().GetBrush(IconStyleName);
+		if (IconBrush && IconBrush->GetResourceName() != NAME_None)
+		{
+			const float IconSize = FMath::Min(16.0f, HeaderSize.Y - 4.0f);
+			const FVector2D IconPosition(2.0f, 2.0f);
+
+			FSlateDrawElement::MakeBox(
+				Painter.DrawElements,
+				LayerId,
+				Painter.HeaderGeometry.ToPaintGeometry(FVector2D(IconSize, IconSize), FSlateLayoutTransform(IconPosition)),
+				IconBrush,
+				ESlateDrawEffect::None,
+				FLinearColor::White);
+		}
+	}
+
+	return LayerId + 1;
+}
+
+void FAnimTransitionSectionInterface::GenerateSectionLayout(ISectionLayoutBuilder& LayoutBuilder)
+{
+	// Call the base implementation to generate channel layouts for expansion.
+	// The blend curve channel has bCanCollapseToTrack = false, so it won't add keys
+	// to the section header - instead it creates an expandable row.
+	// We draw our own curve visualization in OnPaintSection.
+	ISequencerSection::GenerateSectionLayout(LayoutBuilder);
+}
+
+void FAnimTransitionSectionInterface::BuildSectionContextMenu(FMenuBuilder& MenuBuilder, const FGuid& ObjectBinding)
+{
+	MenuBuilder.BeginSection(NAME_None, LOCTEXT("TransitionSection", "Transition"));
+	{
+		MenuBuilder.AddSubMenu(
+			LOCTEXT("ChangeTransitionType_Label", "Change Transition Type"),
+			LOCTEXT("ChangeTransitionType_Tooltip", "Change the type of this transition"),
+			FNewMenuDelegate::CreateSP(this, &FAnimTransitionSectionInterface::PopulateChangeTypeMenu)
+		);
+	}
+	MenuBuilder.EndSection();
+}
+
+void FAnimTransitionSectionInterface::PopulateChangeTypeMenu(FMenuBuilder& MenuBuilder)
+{
+	UMovieSceneAnimTransitionSectionBase* TransitionSection = GetTransitionSection();
+	if (!TransitionSection)
+	{
+		return;
+	}
+
+	// Get all transition section classes
+	TArray<UClass*> TransitionClasses;
+	for (TObjectIterator<UClass> It; It; ++It)
+	{
+		UClass* Class = *It;
+		if (Class->IsChildOf(UMovieSceneAnimTransitionSectionBase::StaticClass()) &&
+			!Class->HasAnyClassFlags(CLASS_Abstract) &&
+			Class != UMovieSceneAnimTransitionSectionBase::StaticClass())
+		{
+			TransitionClasses.Add(Class);
+		}
+	}
+
+	// Sort by display name
+	TransitionClasses.Sort([](const UClass& A, const UClass& B)
+	{
+		return A.GetDisplayNameText().CompareTo(B.GetDisplayNameText()) < 0;
+	});
+
+	UClass* CurrentClass = TransitionSection->GetClass();
+
+	for (UClass* TransitionClass : TransitionClasses)
+	{
+		FText DisplayName = TransitionClass->GetDisplayNameText();
+		bool bIsCurrentType = (TransitionClass == CurrentClass);
+
+		MenuBuilder.AddMenuEntry(
+			DisplayName,
+			FText::Format(LOCTEXT("ChangeToType_Tooltip", "Change to {0}"), DisplayName),
+			FSlateIcon(),
+			FUIAction(
+				FExecuteAction::CreateSP(this, &FAnimTransitionSectionInterface::ChangeTransitionType, TransitionClass),
+				FCanExecuteAction::CreateLambda([bIsCurrentType] { return !bIsCurrentType; }),
+				FIsActionChecked::CreateLambda([bIsCurrentType] { return bIsCurrentType; })
+			),
+			NAME_None,
+			EUserInterfaceActionType::RadioButton
+		);
+	}
+}
+
+void FAnimTransitionSectionInterface::BeginResizeSection()
+{
+	UMovieSceneAnimTransitionSectionBase* TransitionSection = GetTransitionSection();
+	if (!TransitionSection)
+	{
+		return;
+	}
+
+	// Cache initial ranges for constrained resize
+	InitialTransitionRange = TransitionSection->GetRange();
+
+	if (TransitionSection->FromSection)
+	{
+		InitialFromRange = TransitionSection->FromSection->GetRange();
+	}
+	else
+	{
+		InitialFromRange = TRange<FFrameNumber>::Empty();
+	}
+
+	if (TransitionSection->ToSection)
+	{
+		InitialToRange = TransitionSection->ToSection->GetRange();
+	}
+	else
+	{
+		InitialToRange = TRange<FFrameNumber>::Empty();
+	}
+}
+
+void FAnimTransitionSectionInterface::ResizeSection(ESequencerSectionResizeMode ResizeMode, FFrameNumber ResizeFrameNumber)
+{
+	UMovieSceneAnimTransitionSectionBase* TransitionSection = GetTransitionSection();
+	if (!TransitionSection)
+	{
+		return;
+	}
+
+	if (!TransitionSection->FromSection || !TransitionSection->ToSection)
+	{
+		return;
+	}
+
+	// Get the old range before any modifications for key scaling
+	TRange<FFrameNumber> OldRange = TransitionSection->GetRange();
+
+	// Constrain resize based on from/to section bounds
+	// The transition defines the overlap region, which must stay within the valid range
+	// where both sections can exist
+
+	if (ResizeMode == SSRM_LeadingEdge)
+	{
+		// Resizing the left edge - this controls where the TO section starts and transition starts
+		// Cannot go earlier than FROM section's original start
+		FFrameNumber MinFrame = InitialFromRange.HasLowerBound() ? InitialFromRange.GetLowerBoundValue() : TNumericLimits<int32>::Min();
+		FFrameNumber MaxFrame = TransitionSection->HasEndFrame() ? TransitionSection->GetExclusiveEndFrame() - 1 : TNumericLimits<int32>::Max();
+
+		FFrameNumber ClampedFrame = FMath::Clamp(ResizeFrameNumber, MinFrame, MaxFrame);
+
+		// Modify TO section's start bound
+		TransitionSection->ToSection->Modify();
+		TransitionSection->ToSection->SetStartFrame(TRangeBound<FFrameNumber>::Inclusive(ClampedFrame));
+
+		// Update transition to match the new overlap
+		TransitionSection->Modify();
+		TransitionSection->SetStartFrame(TRangeBound<FFrameNumber>::Inclusive(ClampedFrame));
+	}
+	else // SSRM_TrailingEdge
+	{
+		// Resizing the right edge - this controls where the FROM section ends and transition ends
+		// Cannot go later than TO section's original end
+		FFrameNumber MinFrame = TransitionSection->HasStartFrame() ? TransitionSection->GetInclusiveStartFrame() + 1 : TNumericLimits<int32>::Min();
+		FFrameNumber MaxFrame = InitialToRange.HasUpperBound() ? InitialToRange.GetUpperBoundValue() : TNumericLimits<int32>::Max();
+
+		FFrameNumber ClampedFrame = FMath::Clamp(ResizeFrameNumber, MinFrame, MaxFrame);
+
+		// Modify FROM section's end bound
+		TransitionSection->FromSection->Modify();
+		TransitionSection->FromSection->SetEndFrame(TRangeBound<FFrameNumber>::Exclusive(ClampedFrame));
+
+		// Update transition to match the new overlap
+		TransitionSection->Modify();
+		TransitionSection->SetEndFrame(TRangeBound<FFrameNumber>::Exclusive(ClampedFrame));
+	}
+
+	// Scale existing keys proportionally to the new range
+	TransitionSection->ScaleKeysToRange(OldRange, TransitionSection->GetRange());
+}
+
+void FAnimTransitionSectionInterface::ChangeTransitionType(UClass* NewTransitionClass)
+{
+	UMovieSceneAnimTransitionSectionBase* OldTransition = GetTransitionSection();
+	UMovieSceneAnimationMixerTrack* MixerTrack = Cast<UMovieSceneAnimationMixerTrack>(TrackPtr.Get());
+
+	if (!OldTransition || !MixerTrack || !NewTransitionClass)
+	{
+		return;
+	}
+
+	// Don't change if already the same type
+	if (OldTransition->GetClass() == NewTransitionClass)
+	{
+		return;
+	}
+
+	FScopedTransaction Transaction(LOCTEXT("ChangeTransitionType", "Change Transition Type"));
+
+	// Store info from old transition
+	UMovieSceneSection* FromSection = OldTransition->FromSection;
+	UMovieSceneSection* ToSection = OldTransition->ToSection;
+	TRange<FFrameNumber> Range = OldTransition->GetRange();
+	int32 RowIndex = OldTransition->GetRowIndex();
+
+	// Get the old blend weight channel if it exists
+	const FMovieSceneFloatChannel* OldBlendChannel = OldTransition->GetBlendWeightChannel();
+
+	// Remove old transition
+	MixerTrack->Modify();
+	MixerTrack->RemoveSection(*OldTransition);
+
+	UMovieSceneAnimTransitionSectionBase* NewTransition = MixerTrack->CreateTransitionSectionOfType(FromSection, ToSection, NewTransitionClass);
+	if (!NewTransition)
+	{
+		// Failed to create new transition - sections may no longer overlap
+		// The old transition has been removed, so just update our pointer to null
+		SectionPtr = nullptr;
+		return;
+	}
+
+	// Copy the blend curve from the old transition if both have blend weight channels
+	const FMovieSceneFloatChannel* NewBlendChannel = NewTransition->GetBlendWeightChannel();
+	if (OldBlendChannel && NewBlendChannel)
+	{
+		// Get non-const pointer to copy to
+		FMovieSceneFloatChannel* MutableNewChannel = const_cast<FMovieSceneFloatChannel*>(NewBlendChannel);
+		*MutableNewChannel = *OldBlendChannel;
+	}
+	else
+	{
+		// No existing curve to copy, initialize defaults
+		NewTransition->InitializeDefaultCurve();
+	}
+
+	// Update the section pointer
+	SectionPtr = NewTransition;
+}
+
+} // namespace UE::Sequencer
+
+#undef LOCTEXT_NAMESPACE
diff --git a/Engine/Plugins/Experimental/MovieSceneAnimMixer/Source/MovieSceneAnimMixerEditor/Private/AnimTransitionSectionInterface.h b/Engine/Plugins/Experimental/MovieSceneAnimMixer/Source/MovieSceneAnimMixerEditor/Private/AnimTransitionSectionInterface.h
new file mode 100644
--- /dev/null
+++ b/Engine/Plugins/Experimental/MovieSceneAnimMixer/Source/MovieSceneAnimMixerEditor/Private/AnimTransitionSectionInterface.h
@@ -0,0 +1,68 @@
+// Copyright Epic Games, Inc. All Rights Reserved.
+
+#pragma once
+
+#include "ISequencerSection.h"
+#include "Templates/SharedPointer.h"
+#include "UObject/WeakObjectPtr.h"
+#include "UObject/WeakObjectPtrTemplates.h"
+
+class UMovieSceneAnimTransitionSectionBase;
+class FSequencerSectionPainter;
+class FMenuBuilder;
+class ISequencer;
+
+namespace UE::Sequencer
+{
+
+/**
+ * Section interface for animation mixer transition sections.
+ * Provides custom rendering and interaction constraints for transitions.
+ */
+class FAnimTransitionSectionInterface
+	: public ISequencerSection
+	, public TSharedFromThis<FAnimTransitionSectionInterface>
+{
+public:
+	FAnimTransitionSectionInterface(UMovieSceneSection& InSection, UMovieSceneTrack& InTrack, FGuid InObjectBinding);
+
+	//~ ISequencerSection interface
+	virtual UMovieSceneSection* GetSectionObject() override;
+	virtual FText GetSectionTitle() const override;
+	virtual FText GetSectionToolTip() const override;
+	virtual int32 OnPaintSection(FSequencerSectionPainter& Painter) const override;
+	virtual void GenerateSectionLayout(class ISectionLayoutBuilder& LayoutBuilder) override;
+	virtual bool IsReadOnly() const override { return false; }
+	virtual void BuildSectionContextMenu(FMenuBuilder& MenuBuilder, const FGuid& ObjectBinding) override;
+
+	/** Custom resize behavior: resizing the transition adjusts the from/to sections */
+	virtual void BeginResizeSection() override;
+	virtual void ResizeSection(ESequencerSectionResizeMode ResizeMode, FFrameNumber ResizeFrameNumber) override;
+
+	/** Transition sections cannot be resized beyond the from section's start or to section's end */
+	virtual bool SectionIsResizable() const override { return true; }
+
+	/** Transition sections cannot be explicitly deleted - they are removed automatically when overlap goes away */
+	virtual bool IsDeletable() const override { return false; }
+
+	/** Transition sections draw their own curve visualization in OnPaintSection, so don't show the key area when collapsed */
+	virtual bool ShouldShowKeyAreaWhenCollapsed() const override { return false; }
+
+protected:
+	UMovieSceneAnimTransitionSectionBase* GetTransitionSection() const;
+
+private:
+	void PopulateChangeTypeMenu(FMenuBuilder& MenuBuilder);
+	void ChangeTransitionType(UClass* NewTransitionClass);
+
+	TWeakObjectPtr<UMovieSceneSection> SectionPtr;
+	TWeakObjectPtr<UMovieSceneTrack> TrackPtr;
+	FGuid ObjectBindingId;
+
+	// Cached state for resize operations
+	TRange<FFrameNumber> InitialFromRange;
+	TRange<FFrameNumber> InitialToRange;
+	TRange<FFrameNumber> InitialTransitionRange;
+};
+
+} // namespace UE::Sequencer
diff --git a/Engine/Plugins/Experimental/MovieSceneAnimMixer/Source/MovieSceneAnimMixerEditor/Private/MVVM/AnimationMixerLayerModel.cpp b/Engine/Plugins/Experimental/MovieSceneAnimMixer/Source/MovieSceneAnimMixerEditor/Private/MVVM/AnimationMixerLayerModel.cpp
--- a/Engine/Plugins/Experimental/MovieSceneAnimMixer/Source/MovieSceneAnimMixerEditor/Private/MVVM/AnimationMixerLayerModel.cpp
+++ b/Engine/Plugins/Experimental/MovieSceneAnimMixer/Source/MovieSceneAnimMixerEditor/Private/MVVM/AnimationMixerLayerModel.cpp
@@ -5,7 +5,10 @@
 #include "MVVM/AnimationMixerLayoutHelpers.h"
 #include "MovieSceneAnimationMixerTrack.h"
 #include "MovieSceneAnimationMixerLayer.h"
+#include "MovieSceneAnimTransitionSectionBase.h"
 #include "MovieSceneTrack.h"
+#include "MovieScene.h"
+#include "MovieSceneSequence.h"
 #include "ISequencerModule.h"
 #include "Modules/ModuleManager.h"
 
@@ -24,6 +27,7 @@
 #include "MVVM/TrackRowModelStorageExtension.h"
 #include "MVVM/Extensions/IRecyclableExtension.h"
 #include "MVVM/Extensions/IObjectBindingExtension.h"
+#include "MVVM/Extensions/IOutlinerExtension.h"
 #include "ISequencer.h"
 #include "ISequencerTrackEditor.h"
 #include "SequencerUtilities.h"
@@ -39,6 +43,71 @@ namespace Sequencer
 
 UE_SEQUENCER_DEFINE_CASTABLE(FAnimationMixerLayerModel);
 
+/**
+ * Sort sections for display: transitions appear between their from/to sections.
+ * Regular sections are sorted by start time, and transitions are inserted after their "from" section.
+ */
+static void SortSectionsForDisplay(TArray<UMovieSceneSection*>& Sections)
+{
+	// Separate transitions from regular sections
+	TArray<UMovieSceneAnimTransitionSectionBase*> Transitions;
+	TArray<UMovieSceneSection*> RegularSections;
+
+	for (UMovieSceneSection* Section : Sections)
+	{
+		if (UMovieSceneAnimTransitionSectionBase* Transition = Cast<UMovieSceneAnimTransitionSectionBase>(Section))
+		{
+			Transitions.Add(Transition);
+		}
+		else if (Section)
+		{
+			RegularSections.Add(Section);
+		}
+	}
+
+	// Sort regular sections by start time
+	RegularSections.Sort([](const UMovieSceneSection& A, const UMovieSceneSection& B)
+	{
+		if (!A.HasStartFrame() && !B.HasStartFrame())
+		{
+			return false;
+		}
+		if (!A.HasStartFrame())
+		{
+			return true;
+		}
+		if (!B.HasStartFrame())
+		{
+			return false;
+		}
+		return A.GetInclusiveStartFrame() < B.GetInclusiveStartFrame();
+	});
+
+	// Build result: insert transitions after their "from" section
+	Sections.Empty(RegularSections.Num() + Transitions.Num());
+
+	for (UMovieSceneSection* Section : RegularSections)
+	{
+		Sections.Add(Section);
+
+		// Check if any transitions have this section as their "from" section
+		for (int32 i = Transitions.Num() - 1; i >= 0; --i)
+		{
+			if (Transitions[i]->FromSection == Section)
+			{
+				Sections.Add(Transitions[i]);
+				Transitions.RemoveAt(i);
+			}
+		}
+	}
+
+	// Add any remaining transitions (orphaned or with invalid from/to references)
+	for (UMovieSceneAnimTransitionSectionBase* Transition : Transitions)
+	{
+		Sections.Add(Transition);
+	}
+}
+
 FAnimationMixerLayerModel::FAnimationMixerLayerModel(UMovieSceneAnimationMixerTrack* InParentTrack, UMovieSceneAnimationMixerLayer* InLayer)
 	: WeakParentTrack(InParentTrack)
 	, WeakLayer(InLayer)
@@ -155,7 +224,9 @@ void FAnimationMixerLayerModel::RefreshLayout(bool bChildrenNeedLayout)
 		{
 			// Single row case - use helper for section layout
 			TArray<UMovieSceneSection*> AllSections = ChildTrack->GetAllSections();
-			bool bCreateOutliners = bIsExpanded && AllSections.Num() > 1;
+			// Query saved expansion state directly - we can't use IsExpanded() here because
+			// it checks for children, but we're in the process of building them
+			bool bCreateOutliners = GetSavedExpansionState() && AllSections.Num() > 1;
 
 			// Only recycle top level channels if children need layout
 			if (bChildrenNeedLayout)
@@ -224,7 +295,13 @@ void FAnimationMixerLayerModel::RefreshLayout(bool bChildrenNeedLayout)
 		{
 			// Get sections from the layer and use helper
 			TArray<UMovieSceneSection*> LayerSections = Layer->GetSections();
-			bool bCreateOutliners = bIsExpanded && LayerSections.Num() > 1;
+
+			// Sort sections so transitions appear between their from/to sections
+			SortSectionsForDisplay(LayerSections);
+
+			// Query saved expansion state directly - we can't use IsExpanded() here because
+			// it checks for children, but we're in the process of building them
+			bool bCreateOutliners = GetSavedExpansionState() && LayerSections.Num() > 1;
 
 			// Only recycle top level channels if children need layout
 			if (bChildrenNeedLayout)
@@ -986,9 +1063,15 @@ void FAnimationMixerLayerModel::BuildSidebarMenu(FMenuBuilder& MenuBuilder)
 	FOutlinerItemModel::BuildSidebarMenu(MenuBuilder);
 }
 
-void FAnimationMixerLayerModel::SetExpansion(bool bInIsExpanded)  
-{ 
+void FAnimationMixerLayerModel::SetExpansion(bool bInIsExpanded)
+{
 	FEvaluableOutlinerItemModel::SetExpansion(bInIsExpanded);
+
+	// Reset layout state to force complete rebuild when expansion changes
+	// This ensures we properly switch between collapsed (sections in SectionList)
+	// and expanded (sections in section outliners) modes
+	PreviousLayoutNumSections.Reset();
+
 	RefreshLayout(true);
 }
 
@@ -1133,6 +1216,33 @@ void FAnimationMixerLayerModel::SetIsSoloed(bool bIsSoloed)
 	}
 }
 
+bool FAnimationMixerLayerModel::GetSavedExpansionState() const
+{
+	// This replicates the initialization logic from FOutlinerItemModelMixin::IsExpanded()
+	// but without the children check at the end. We need this during layout because
+	// IsExpanded() checks for children, but we're in the process of building them.
+
+	TStringBuilder<256> StringBuilder;
+	IOutlinerExtension::GetPathName(*this, StringBuilder);
+
+	TSharedPtr<FSequenceModel> SequenceModel = FindAncestorOfType<FSequenceModel>();
+	UMovieSceneSequence* Sequence = SequenceModel ? SequenceModel->GetSequence() : nullptr;
+	UMovieScene* MovieScene = Sequence ? Sequence->GetMovieScene() : nullptr;
+
+	if (MovieScene)
+	{
+		FStringView StringView = StringBuilder.ToView();
+		FMovieSceneEditorData& EditorData = MovieScene->GetEditorData();
+		if (const FMovieSceneExpansionState* Expansion = EditorData.ExpansionStates.FindByHash(GetTypeHash(StringView), StringView))
+		{
+			return Expansion->bExpanded;
+		}
+	}
+
+	// Fall back to default expansion state if no saved state exists
+	return GetDefaultExpansionState();
+}
+
 } // namespace Sequencer
 } // namespace UE
 
diff --git a/Engine/Plugins/Experimental/MovieSceneAnimMixer/Source/MovieSceneAnimMixerEditor/Private/MVVM/AnimationMixerTrackModel.cpp b/Engine/Plugins/Experimental/MovieSceneAnimMixer/Source/MovieSceneAnimMixerEditor/Private/MVVM/AnimationMixerTrackModel.cpp
--- a/Engine/Plugins/Experimental/MovieSceneAnimMixer/Source/MovieSceneAnimMixerEditor/Private/MVVM/AnimationMixerTrackModel.cpp
+++ b/Engine/Plugins/Experimental/MovieSceneAnimMixer/Source/MovieSceneAnimMixerEditor/Private/MVVM/AnimationMixerTrackModel.cpp
@@ -5,6 +5,7 @@
 #include "MVVM/AnimationMixerLayoutHelpers.h"
 #include "MovieSceneAnimationMixerTrack.h"
 #include "MovieSceneAnimationMixerLayer.h"
+#include "MovieSceneAnimTransitionSectionBase.h"
 #include "MVVM/ViewModels/TrackModel.h"
 #include "MVVM/ViewModels/TrackRowModel.h"
 #include "MVVM/ViewModels/SequenceModel.h"
@@ -359,6 +360,15 @@ TOptional<FAnimMixerDropPreview> FAnimationMixerTrackModel::ComputeDropPreview(c
 		return TOptional<FAnimMixerDropPreview>();
 	}
 
+	// Transition sections cannot be moved vertically - they must stay on the same row as from/to sections
+	for (UMovieSceneSection* Section : *Context.DraggedSections)
+	{
+		if (Section && Section->IsA<UMovieSceneAnimTransitionSectionBase>())
+		{
+			return TOptional<FAnimMixerDropPreview>();
+		}
+	}
+
 	// Compute combined frame range of all dragged sections (for outline drawing)
 	TRange<FFrameNumber> DraggedSectionRange = TRange<FFrameNumber>::Empty();
 	for (UMovieSceneSection* Section : *Context.DraggedSections)
@@ -562,6 +572,16 @@ bool FAnimationMixerTrackModel::ApplyDropPreview(const FAnimMixerDropPreview& Pr
 
 	const TArray<TObjectPtr<UMovieSceneAnimationMixerLayer>>& Layers = MixerTrack->GetLayers();
 
+	// Track all affected row indices for transition updates
+	TSet<int32> AffectedRows;
+
+	// For non-insert mode, add target row now. For insert mode, we add it after shift adjustment
+	// because the destination row index doesn't change (it's the newly inserted row).
+	if (!bInsertMode)
+	{
+		AffectedRows.Add(TargetRowIndex);
+	}
+
 	if (bInsertMode)
 	{
 		// Remove sections from source layers before shifting
@@ -570,6 +590,7 @@ bool FAnimationMixerTrackModel::ApplyDropPreview(const FAnimMixerDropPreview& Pr
 			if (Section)
 			{
 				int32 SourceRowIndex = Section->GetRowIndex();
+				AffectedRows.Add(SourceRowIndex);
 				if (SourceRowIndex >= 0 && SourceRowIndex < Layers.Num())
 				{
 					UMovieSceneAnimationMixerLayer* SourceLayer = Layers[SourceRowIndex];
@@ -606,6 +627,25 @@ bool FAnimationMixerTrackModel::ApplyDropPreview(const FAnimMixerDropPreview& Pr
 				}
 			}
 		}
+
+		// After shift, source rows have shifted - update affected rows to post-shift indices
+		// All source rows at or above TargetRowIndex have been shifted down by 1
+		TSet<int32> ShiftedAffectedRows;
+		for (int32 RowIndex : AffectedRows)
+		{
+			if (RowIndex >= TargetRowIndex)
+			{
+				ShiftedAffectedRows.Add(RowIndex + 1);
+			}
+			else
+			{
+				ShiftedAffectedRows.Add(RowIndex);
+			}
+		}
+		AffectedRows = MoveTemp(ShiftedAffectedRows);
+
+		// Now add the destination row - it doesn't shift because it's the newly inserted row
+		AffectedRows.Add(TargetRowIndex);
 	}
 	else
 	{
@@ -622,6 +662,7 @@ bool FAnimationMixerTrackModel::ApplyDropPreview(const FAnimMixerDropPreview& Pr
 
 						// Remove from source layer
 						int32 SourceRowIndex = Section->GetRowIndex();
+						AffectedRows.Add(SourceRowIndex);
 						if (SourceRowIndex >= 0 && SourceRowIndex < Layers.Num())
 						{
 							if (UMovieSceneAnimationMixerLayer* SourceLayer = Layers[SourceRowIndex])
@@ -644,6 +685,9 @@ bool FAnimationMixerTrackModel::ApplyDropPreview(const FAnimMixerDropPreview& Pr
 		}
 	}
 
+	// Update transitions for all affected rows
+	MixerTrack->UpdateTransitionsForRows(AffectedRows);
+
 	// Force UI update after applying the drop
 	ForceUpdate();
 
diff --git a/Engine/Plugins/Experimental/MovieSceneAnimMixer/Source/MovieSceneAnimMixerEditor/Private/MovieSceneAnimMixerEditorModule.cpp b/Engine/Plugins/Experimental/MovieSceneAnimMixer/Source/MovieSceneAnimMixerEditor/Private/MovieSceneAnimMixerEditorModule.cpp
--- a/Engine/Plugins/Experimental/MovieSceneAnimMixer/Source/MovieSceneAnimMixerEditor/Private/MovieSceneAnimMixerEditorModule.cpp
+++ b/Engine/Plugins/Experimental/MovieSceneAnimMixer/Source/MovieSceneAnimMixerEditor/Private/MovieSceneAnimMixerEditorModule.cpp
@@ -4,6 +4,7 @@
 
 #include "AnimBlueprintExtension_SequencerMixerTarget.h"
 #include "AnimMixerTargetMenuProviders.h"
+#include "AnimTransitionSectionInterface.h"
 #include "IAnimGraph_SequencerMixerTargetConnector.h"
 #include "IMovieSceneAnimMixerTargetMenuProvider.h"
 #include "ISequencerModule.h"
@@ -11,6 +12,8 @@
 #include "Modules/ModuleManager.h"
 #include "MovieSceneAnimationMixerTrack.h"
 #include "MovieSceneAnimationMixerTrackEditor.h"
+#include "MovieSceneAnimCrossfadeTransitionSection.h"
+#include "MovieSceneAnimInertialDeadBlendTransitionSection.h"
 #include "Channels/BuiltInChannelEditors.h"
 #include "MVVM/AnimationMixerTrackModel.h"
 
@@ -107,6 +110,21 @@ namespace UE::MovieScene
 
 		SequencerModule.RegisterChannelInterface<FMovieSceneByteChannelDefaultOnly>();
 
+		// Register custom section interface for transition sections
+		UE::Sequencer::FAnimationMixerTrackEditor::RegisterCustomMixerAnimSection(
+			UMovieSceneAnimCrossfadeTransitionSection::StaticClass(),
+			FOnMakeSectionInterfaceDelegate::CreateLambda([](UMovieSceneSection& Section, UMovieSceneTrack& Track, FGuid ObjectBinding) -> TSharedRef<ISequencerSection>
+			{
+				return MakeShared<UE::Sequencer::FAnimTransitionSectionInterface>(Section, Track, ObjectBinding);
+			}));
+
+		UE::Sequencer::FAnimationMixerTrackEditor::RegisterCustomMixerAnimSection(
+			UMovieSceneAnimInertialDeadBlendTransitionSection::StaticClass(),
+			FOnMakeSectionInterfaceDelegate::CreateLambda([](UMovieSceneSection& Section, UMovieSceneTrack& Track, FGuid ObjectBinding) -> TSharedRef<ISequencerSection>
+			{
+				return MakeShared<UE::Sequencer::FAnimTransitionSectionInterface>(Section, Track, ObjectBinding);
+			}));
+
 		// Register target menu providers
 		AutomaticTargetMenuProvider = MakeUnique<FAutomaticTargetMenuProvider>();
 		AnimInstanceTargetMenuProvider = MakeUnique<FAnimInstanceTargetMenuProvider>();
@@ -139,6 +157,10 @@ namespace UE::MovieScene
 			IModularFeatures::Get().UnregisterModularFeature(IMovieSceneAnimMixerTargetMenuProvider::GetModularFeatureName(), AnimNextInjectionTargetMenuProvider.Get());
 		}
 
+		// Unregister custom section interface for transition sections
+		UE::Sequencer::FAnimationMixerTrackEditor::UnregisterCustomMixerAnimSection(UMovieSceneAnimCrossfadeTransitionSection::StaticClass());
+		UE::Sequencer::FAnimationMixerTrackEditor::UnregisterCustomMixerAnimSection(UMovieSceneAnimInertialDeadBlendTransitionSection::StaticClass());
+
 		if (ISequencerModule* SequencerModule = FModuleManager::Get().GetModulePtr<ISequencerModule>(SequencerModuleName.Resolve()))
 		{
 			SequencerModule->UnregisterTrackModel(AnimationMixerTrackModelHandle);
diff --git a/Engine/Plugins/Experimental/MovieSceneAnimMixer/Source/MovieSceneAnimMixerEditor/Private/MovieSceneAnimationMixerTrackEditor.cpp b/Engine/Plugins/Experimental/MovieSceneAnimMixer/Source/MovieSceneAnimMixerEditor/Private/MovieSceneAnimationMixerTrackEditor.cpp
--- a/Engine/Plugins/Experimental/MovieSceneAnimMixer/Source/MovieSceneAnimMixerEditor/Private/MovieSceneAnimationMixerTrackEditor.cpp
+++ b/Engine/Plugins/Experimental/MovieSceneAnimMixer/Source/MovieSceneAnimMixerEditor/Private/MovieSceneAnimationMixerTrackEditor.cpp
@@ -3,6 +3,7 @@
 #include "MovieSceneAnimationMixerTrackEditor.h"
 
 #include "Animation/AnimSequenceBase.h"
+#include "UObject/UObjectIterator.h"
 #include "AssetRegistry/AssetRegistryModule.h"
 #include "ClassIconFinder.h"
 #include "Components/SkeletalMeshComponent.h"
@@ -11,6 +12,7 @@
 #include "GameFramework/Actor.h"
 #include "IMovieSceneAnimMixerItemMenuProvider.h"
 #include "IMovieSceneAnimMixerTargetMenuProvider.h"
+#include "AnimatedRange.h"
 #include "ISequencer.h"
 #include "MovieSceneAnimationMixerItemInterface.h"
 #include "MovieSceneAnimationMixerLayer.h"
@@ -26,7 +28,11 @@
 #include "MVVM/Views/ViewUtilities.h"
 #include "ScopedTransaction.h"
 #include "Sections/MovieSceneSkeletalAnimationSection.h"
+#include "MovieSceneAnimTransitionSectionBase.h"
+#include "MVVM/ViewModels/SectionOutlinerModel.h"
+#include "SequencerSectionPainter.h"
 #include "SequencerUtilities.h"
+#include "TimeToPixel.h"
 #include "Widgets/Input/SComboButton.h"
 #include "Widgets/Images/SImage.h"
 #include "MVVM/AnimationMixerTrackModel.h"
@@ -192,6 +198,81 @@ public:
 		: FCommonAnimationSection(InSection, InSequencer)
 	{}
 
+	virtual int32 OnPaintSection(FSequencerSectionPainter& Painter) const override
+	{
+		// Check if we're in section outliner mode (layer expanded, each section on its own row)
+		// In outliner mode, sections don't visually overlap so no transition padding is needed
+		const bool bIsInOutlinerMode = Painter.SectionModel &&
+			Painter.SectionModel->FindAncestorOfType<UE::Sequencer::FSectionOutlinerModel>().IsValid();
+
+		// Calculate transition padding using the accurate time converter before painting
+		// This value will be used by GetContentPadding() which is called after OnPaintSection()
+		CachedTransitionPaddingLeft = bIsInOutlinerMode ? 0.0f : CalculateTransitionPadding(Painter.GetTimeConverter());
+
+		return FCommonAnimationSection::OnPaintSection(Painter);
+	}
+
+	virtual FMargin GetContentPadding() const override
+	{
+		FMargin BasePadding = FCommonAnimationSection::GetContentPadding();
+
+		// Add the transition padding that was calculated in OnPaintSection
+		BasePadding.Left += CachedTransitionPaddingLeft;
+
+		return BasePadding;
+	}
+
+private:
+
+	float CalculateTransitionPadding(const FTimeToPixel& TimeConverter) const
+	{
+		// Check if there's a transition section overlapping the start of this section
+		// If so, calculate extra left padding so the text appears after the transition
+		UMovieSceneSection* Section = WeakSection.Get();
+		if (!Section || !Section->HasStartFrame())
+		{
+			return 0.0f;
+		}
+
+		UMovieSceneAnimationMixerTrack* MixerTrack = Cast<UMovieSceneAnimationMixerTrack>(Section->GetTypedOuter<UMovieSceneTrack>());
+		if (!MixerTrack)
+		{
+			return 0.0f;
+		}
+
+		// Find any transition that references this section as the "to" section
+		TArray<UMovieSceneAnimTransitionSectionBase*> Transitions = MixerTrack->FindTransitionsForSection(Section);
+		for (UMovieSceneAnimTransitionSectionBase* Transition : Transitions)
+		{
+			if (Transition && Transition->ToSection == Section)
+			{
+				// This is the "to" section of a transition - calculate padding for the transition width
+				TRange<FFrameNumber> TransitionRange = Transition->GetRange();
+				if (!TransitionRange.IsEmpty() && TransitionRange.HasLowerBound() && TransitionRange.HasUpperBound())
+				{
+					// Use the accurate time converter to calculate pixel width
+					FFrameNumber TransitionStart = TransitionRange.GetLowerBoundValue();
+					FFrameNumber TransitionEnd = TransitionRange.GetUpperBoundValue();
+
+					float StartPixel = TimeConverter.FrameToPixel(TransitionStart);
+					float EndPixel = TimeConverter.FrameToPixel(TransitionEnd);
+					float TransitionPixels = EndPixel - StartPixel;
+
+					// Add padding with a small margin
+					return TransitionPixels + 4.0f;
+				}
+				break;  // Only need to handle one transition at the start
+			}
+		}
+
+		return 0.0f;
+	}
+
+	// Cached transition padding calculated in OnPaintSection for use in GetContentPadding
+	mutable float CachedTransitionPaddingLeft = 0.0f;
+
+public:
+
 	virtual void BuildSectionContextMenu(FMenuBuilder& MenuBuilder, const FGuid& ObjectBinding) override
 	{
 		FCommonAnimationSection::BuildSectionContextMenu(MenuBuilder, ObjectBinding);
@@ -205,6 +286,25 @@ public:
 				FNewMenuDelegate::CreateSP(this, &FAnimMixerAnimationSection::PopulateRootTransformMenu)
 			);
 		}
+
+		// Add "Create Transition To..." submenu for overlapping sections
+		UMovieSceneSection* Section = WeakSection.Get();
+		if (Section)
+		{
+			UMovieSceneAnimationMixerTrack* MixerTrack = Cast<UMovieSceneAnimationMixerTrack>(Section->GetTypedOuter<UMovieSceneTrack>());
+			if (MixerTrack)
+			{
+				TArray<UMovieSceneSection*> OverlappingSections = MixerTrack->FindOverlappingSections(Section);
+				if (OverlappingSections.Num() > 0)
+				{
+					MenuBuilder.AddSubMenu(
+						LOCTEXT("CreateTransition_Label", "Create Transition To..."),
+						LOCTEXT("CreateTransition_Tooltip", "Create a transition between this section and an overlapping section"),
+						FNewMenuDelegate::CreateSP(this, &FAnimMixerAnimationSection::PopulateCreateTransitionMenu)
+					);
+				}
+			}
+		}
 	}
 
 	UMovieSceneAnimationBaseTransformDecoration* FindRootDecoration() const
@@ -364,6 +464,140 @@ public:
 			Decoration->RootOriginLocation = CurrentTransform.GetTranslation();
 		}
 	}
+
+	void PopulateCreateTransitionMenu(FMenuBuilder& MenuBuilder)
+	{
+		UMovieSceneSection* Section = WeakSection.Get();
+		if (!Section)
+		{
+			return;
+		}
+
+		UMovieSceneAnimationMixerTrack* MixerTrack = Cast<UMovieSceneAnimationMixerTrack>(Section->GetTypedOuter<UMovieSceneTrack>());
+		if (!MixerTrack)
+		{
+			return;
+		}
+
+		TArray<UMovieSceneSection*> OverlappingSections = MixerTrack->FindOverlappingSections(Section);
+
+		// Get all non-abstract transition section classes
+		TArray<UClass*> TransitionClasses;
+		for (TObjectIterator<UClass> It; It; ++It)
+		{
+			UClass* Class = *It;
+			if (Class->IsChildOf(UMovieSceneAnimTransitionSectionBase::StaticClass()) &&
+				!Class->HasAnyClassFlags(CLASS_Abstract) &&
+				Class != UMovieSceneAnimTransitionSectionBase::StaticClass())
+			{
+				TransitionClasses.Add(Class);
+			}
+		}
+
+		// Sort by display name
+		TransitionClasses.Sort([](const UClass& A, const UClass& B)
+		{
+			return A.GetDisplayNameText().CompareTo(B.GetDisplayNameText()) < 0;
+		});
+
+		for (UMovieSceneSection* OtherSection : OverlappingSections)
+		{
+			// Skip sections without start frames (unbounded sections can't have meaningful transitions)
+			if (!Section->HasStartFrame() || !OtherSection->HasStartFrame())
+			{
+				continue;
+			}
+
+			// Determine from/to based on start time
+			UMovieSceneSection* FromSection = Section;
+			UMovieSceneSection* ToSection = OtherSection;
+			if (OtherSection->GetInclusiveStartFrame() < Section->GetInclusiveStartFrame())
+			{
+				FromSection = OtherSection;
+				ToSection = Section;
+			}
+
+			// Check if a transition already exists for this pair
+			UMovieSceneAnimTransitionSectionBase* ExistingTransition = MixerTrack->FindTransitionForPair(FromSection, ToSection);
+			if (ExistingTransition)
+			{
+				continue; // Skip - transition already exists
+			}
+
+			// Get a display name for the other section
+			FText OtherSectionName = FText::FromString(OtherSection->GetName());
+			if (UMovieSceneSkeletalAnimationSection* OtherAnimSection = Cast<UMovieSceneSkeletalAnimationSection>(OtherSection))
+			{
+				if (UAnimSequenceBase* Animation = OtherAnimSection->GetAnimation())
+				{
+					OtherSectionName = FText::FromString(Animation->GetName());
+				}
+			}
+
+			// Add submenu for this overlapping section with transition type options
+			MenuBuilder.AddSubMenu(
+				OtherSectionName,
+				FText::Format(LOCTEXT("CreateTransitionTo_Tooltip", "Create a transition to {0}"), OtherSectionName),
+				FNewMenuDelegate::CreateSP(this, &FAnimMixerAnimationSection::PopulateTransitionTypeMenu, OtherSection, TransitionClasses)
+			);
+		}
+	}
+
+	void PopulateTransitionTypeMenu(FMenuBuilder& MenuBuilder, UMovieSceneSection* OtherSection, TArray<UClass*> TransitionClasses)
+	{
+		UMovieSceneSection* Section = WeakSection.Get();
+		if (!Section || !OtherSection)
+		{
+			return;
+		}
+
+		// Skip sections without start frames
+		if (!Section->HasStartFrame() || !OtherSection->HasStartFrame())
+		{
+			return;
+		}
+
+		// Determine from/to based on start time
+		UMovieSceneSection* FromSection = Section;
+		UMovieSceneSection* ToSection = OtherSection;
+		if (OtherSection->GetInclusiveStartFrame() < Section->GetInclusiveStartFrame())
+		{
+			FromSection = OtherSection;
+			ToSection = Section;
+		}
+
+		for (UClass* TransitionClass : TransitionClasses)
+		{
+			FText DisplayName = TransitionClass->GetDisplayNameText();
+
+			MenuBuilder.AddMenuEntry(
+				DisplayName,
+				FText::Format(LOCTEXT("CreateTransitionOfType_Tooltip", "Create a {0} transition"), DisplayName),
+				FSlateIcon(),
+				FUIAction(
+					FExecuteAction::CreateSP(this, &FAnimMixerAnimationSection::CreateTransitionToSection, FromSection, ToSection, TransitionClass)
+				)
+			);
+		}
+	}
+
+	void CreateTransitionToSection(UMovieSceneSection* FromSection, UMovieSceneSection* ToSection, UClass* TransitionClass)
+	{
+		if (!FromSection || !ToSection || !TransitionClass)
+		{
+			return;
+		}
+
+		UMovieSceneAnimationMixerTrack* MixerTrack = Cast<UMovieSceneAnimationMixerTrack>(FromSection->GetTypedOuter<UMovieSceneTrack>());
+		if (!MixerTrack)
+		{
+			return;
+		}
+
+		FScopedTransaction Transaction(LOCTEXT("CreateTransition", "Create Transition"));
+
+		MixerTrack->CreateTransitionSectionOfType(FromSection, ToSection, TransitionClass);
+	}
 };
 
 TMap<const UClass*, FOnMakeSectionInterfaceDelegate> FAnimationMixerTrackEditor::MakeSectionInterfaceCallbacks;
@@ -668,8 +902,6 @@ TSharedRef<SWidget> FAnimationMixerTrackEditor::BuildAddSectionSubMenu(TWeakView
 
 			MixerTrack->AddSection(*NewSection);
 
-			MixerTrack->UpdateEasing();
-
 			Sequencer->NotifyMovieSceneDataChanged(EMovieSceneDataChangeType::MovieSceneStructureItemAdded);
 		}
 	};
@@ -864,6 +1096,7 @@ TSharedRef<SWidget> FAnimationMixerTrackEditor::BuildAddSectionSubMenu(TWeakView
 	}
 
 	// Filter to only classes implementing the mixer item interface
+	// Also exclude transition sections - they are created via the transition context menu, not the add menu
 	for (auto It = AllItemClasses.CreateIterator(); It; ++It)
 	{
 		const UClass* Class = FSoftClassPath(It->ToString()).ResolveClass();
@@ -871,6 +1104,11 @@ TSharedRef<SWidget> FAnimationMixerTrackEditor::BuildAddSectionSubMenu(TWeakView
 		{
 			It.RemoveCurrent();
 		}
+		else if (Class->IsChildOf(UMovieSceneAnimTransitionSectionBase::StaticClass()))
+		{
+			// Transition sections are created via context menu, not the add menu
+			It.RemoveCurrent();
+		}
 	}
 
 	// Find all registered menu providers
diff --git a/Engine/Plugins/Experimental/MovieSceneAnimMixer/Source/MovieSceneAnimMixerEditor/Public/MVVM/AnimationMixerLayerModel.h b/Engine/Plugins/Experimental/MovieSceneAnimMixer/Source/MovieSceneAnimMixerEditor/Public/MVVM/AnimationMixerLayerModel.h
--- a/Engine/Plugins/Experimental/MovieSceneAnimMixer/Source/MovieSceneAnimMixerEditor/Public/MVVM/AnimationMixerLayerModel.h
+++ b/Engine/Plugins/Experimental/MovieSceneAnimMixer/Source/MovieSceneAnimMixerEditor/Public/MVVM/AnimationMixerLayerModel.h
@@ -3,6 +3,7 @@
 #pragma once
 
 #include "MVVM/ViewModels/OutlinerItemModel.h"
+#include "MVVM/ViewModels/TrackModel.h"
 #include "MVVM/Extensions/IRenameableExtension.h"
 #include "MVVM/Extensions/IResizableExtension.h"
 #include "MVVM/Extensions/ISectionOwnerExtension.h"
@@ -186,9 +187,16 @@ public:
 
 private:
 
+	/**
+	 * Query the saved expansion state from editor data without checking for children.
+	 * This is needed during layout because IsExpanded() checks for children, but we're
+	 * in the process of building them.
+	 */
+	bool GetSavedExpansionState() const;
+
 	/** A second children list for the sections inside this layer */
 	FViewModelListHead SectionList{EViewModelListType::TrackArea};
-	FViewModelListHead TopLevelChannelList{EViewModelListType::TrackArea};
+	FViewModelListHead TopLevelChannelList{FTrackModel::GetTopLevelChannelGroupType()};
 
 	/** The parent animation mixer track */
 	TWeakObjectPtr<UMovieSceneAnimationMixerTrack> WeakParentTrack;
diff --git a/Engine/Source/Editor/Sequencer/Private/MVVM/ViewModels/SectionOutlinerModel.cpp b/Engine/Source/Editor/Sequencer/Private/MVVM/ViewModels/SectionOutlinerModel.cpp
--- a/Engine/Source/Editor/Sequencer/Private/MVVM/ViewModels/SectionOutlinerModel.cpp
+++ b/Engine/Source/Editor/Sequencer/Private/MVVM/ViewModels/SectionOutlinerModel.cpp
@@ -2,6 +2,7 @@
 
 #include "MVVM/ViewModels/SectionOutlinerModel.h"
 #include "MVVM/ViewModels/SectionModel.h"
+#include "MVVM/ViewModels/TrackModel.h"
 #include "MVVM/ViewModels/SequencerEditorViewModel.h"
 #include "MVVM/ViewModels/SequenceModel.h"
 #include "MVVM/ViewModels/ChannelModel.h"
@@ -51,7 +52,8 @@ namespace
 UE_SEQUENCER_DEFINE_CASTABLE(FSectionOutlinerModel);
 
 FSectionOutlinerModel::FSectionOutlinerModel(UMovieSceneSection* InSection, TSharedPtr<FSectionModel> InSectionModel)
-	: WeakSection(InSection)
+	: TopLevelChannelList(FTrackModel::GetTopLevelChannelGroupType())
+	, WeakSection(InSection)
 	, WeakSectionModel(InSectionModel)
 {
 	RegisterChildList(&TopLevelChannelList);
@@ -468,7 +470,24 @@ bool FSectionOutlinerModel::IsDimmed() const
 
 bool FSectionOutlinerModel::CanDelete(FText* OutErrorMessage) const
 {
-	return WeakSection.IsValid();
+	if (!WeakSection.IsValid())
+	{
+		return false;
+	}
+
+	// Check if the section interface allows deletion
+	if (TSharedPtr<FSectionModel> SectionModel = WeakSectionModel.Pin())
+	{
+		if (TSharedPtr<ISequencerSection> SectionInterface = SectionModel->GetSectionInterface())
+		{
+			if (!SectionInterface->IsDeletable())
+			{
+				return false;
+			}
+		}
+	}
+
+	return true;
 }
 
 void FSectionOutlinerModel::Delete()
diff --git a/Engine/Source/Editor/Sequencer/Private/SSequencerSection.cpp b/Engine/Source/Editor/Sequencer/Private/SSequencerSection.cpp
--- a/Engine/Source/Editor/Sequencer/Private/SSequencerSection.cpp
+++ b/Engine/Source/Editor/Sequencer/Private/SSequencerSection.cpp
@@ -913,6 +913,12 @@ EVisibility SSequencerSection::GetTopLevelChannelGroupVisibility() const
 		return EVisibility::Collapsed;
 	}
 
+	// Allow sections to opt out of showing the key area when collapsed
+	if (SectionInterface && !SectionInterface->ShouldShowKeyAreaWhenCollapsed())
+	{
+		return EVisibility::Collapsed;
+	}
+
 	return EVisibility::Visible;
 }
 
diff --git a/Engine/Source/Editor/Sequencer/Private/Sequencer.cpp b/Engine/Source/Editor/Sequencer/Private/Sequencer.cpp
--- a/Engine/Source/Editor/Sequencer/Private/Sequencer.cpp
+++ b/Engine/Source/Editor/Sequencer/Private/Sequencer.cpp
@@ -1939,6 +1939,13 @@ void FSequencer::DeleteSections(const TSet<UMovieSceneSection*>& Sections)
 			continue;
 		}
 
+		// Check if section interface allows deletion
+		TSharedPtr<UE::Sequencer::FSectionModel> SectionModel = GetNodeTree()->GetSectionModel(Section);
+		if (SectionModel && SectionModel->GetSectionInterface() && !SectionModel->GetSectionInterface()->IsDeletable())
+		{
+			continue;
+		}
+
 		// if this check fails then the section is outered to a type that doesnt know about the section
 		UMovieSceneTrack* Track = CastChecked<UMovieSceneTrack>(Section->GetOuter());
 		{
@@ -6263,6 +6270,11 @@ bool FSequencer::OnRequestNodeDeleted( TSharedRef<FViewModel> NodeToBeDeleted, c
 
 	if (IDeletableExtension* Deletable = NodeToBeDeleted->CastThis<IDeletableExtension>())
 	{
+		if (!Deletable->CanDelete(nullptr))
+		{
+			return false;
+		}
+
 		Deletable->Delete();
 
 		const TViewModelPtr<IOutlinerExtension> OutlinerItem = NodeToBeDeleted->CastThisShared<IOutlinerExtension>();
diff --git a/Engine/Source/Editor/Sequencer/Private/SequencerContextMenus.cpp b/Engine/Source/Editor/Sequencer/Private/SequencerContextMenus.cpp
--- a/Engine/Source/Editor/Sequencer/Private/SequencerContextMenus.cpp
+++ b/Engine/Source/Editor/Sequencer/Private/SequencerContextMenus.cpp
@@ -532,7 +532,10 @@ void FSectionContextMenu::PopulateMenu(FMenuBuilder& MenuBuilder, TSharedPtr<FEx
 			LOCTEXT("DeleteSection", "Delete"),
 			LOCTEXT("DeleteSectionToolTip", "Deletes this section"),
 			FSlateIcon(),
-			FUIAction(FExecuteAction::CreateLambda([=] { return Shared->DeleteSection(); }))
+			FUIAction(
+				FExecuteAction::CreateLambda([=] { return Shared->DeleteSection(); }),
+				FCanExecuteAction::CreateLambda([=] { return Shared->CanDeleteSection(); })
+			)
 		);
 
 
@@ -1020,6 +1023,8 @@ void FSectionContextMenu::AutoSizeSection()
 
 	FScopedTransaction AutoSizeSectionTransaction(LOCTEXT("AutoSizeSection_Transaction", "Auto Size Section"));
 
+	EMovieSceneSectionMovedResult SectionMovedResult = EMovieSceneSectionMovedResult::None;
+
 	for (UMovieSceneSection* Section : Sequencer->GetViewModel()->GetSelection()->GetSelectedSections())
 	{
 		if (Section && Section->GetAutoSizeRange().IsSet())
@@ -1029,11 +1034,21 @@ void FSectionContextMenu::AutoSizeSection()
 			if (DefaultSectionLength.IsSet())
 			{
 				Section->SetRange(DefaultSectionLength.GetValue());
+
+				// Notify the track that the section was resized so it can update any dependent state
+				if (UMovieSceneTrack* Track = Section->GetTypedOuter<UMovieSceneTrack>())
+				{
+					Track->Modify();
+					SectionMovedResult |= Track->OnSectionMoved(*Section, FMovieSceneSectionMovedParams(EPropertyChangeType::ValueSet));
+				}
 			}
 		}
 	}
 
-	Sequencer->NotifyMovieSceneDataChanged( EMovieSceneDataChangeType::MovieSceneStructureItemAdded );
+	EMovieSceneDataChangeType ChangeType = EnumHasAnyFlags(SectionMovedResult, EMovieSceneSectionMovedResult::SectionsChanged)
+		? EMovieSceneDataChangeType::MovieSceneStructureItemAdded
+		: EMovieSceneDataChangeType::TrackValueChanged;
+	Sequencer->NotifyMovieSceneDataChanged(ChangeType);
 }
 
 
@@ -1361,6 +1376,34 @@ bool FSectionContextMenu::IsSectionLocked() const
 	return true;
 }
 
+bool FSectionContextMenu::CanDeleteSection() const
+{
+	using namespace UE::Sequencer;
+
+	const TSharedPtr<FSequencer> Sequencer = WeakSequencer.Pin();
+	if (!Sequencer.IsValid())
+	{
+		return false;
+	}
+
+	// Check if any selected section is deletable
+	for (UMovieSceneSection* Section : Sequencer->GetViewModel()->GetSelection()->GetSelectedSections())
+	{
+		if (!Section || Section->IsLocked())
+		{
+			continue;
+		}
+
+		// Check if section interface allows deletion
+		TSharedPtr<FSectionModel> SectionModel = Sequencer->GetNodeTree()->GetSectionModel(Section);
+		if (SectionModel && SectionModel->GetSectionInterface() && SectionModel->GetSectionInterface()->IsDeletable())
+		{
+			return true;
+		}
+	}
+
+	return false;
+}
 
 void FSectionContextMenu::DeleteSection()
 {
diff --git a/Engine/Source/Editor/Sequencer/Private/SequencerContextMenus.h b/Engine/Source/Editor/Sequencer/Private/SequencerContextMenus.h
--- a/Engine/Source/Editor/Sequencer/Private/SequencerContextMenus.h
+++ b/Engine/Source/Editor/Sequencer/Private/SequencerContextMenus.h
@@ -82,6 +82,8 @@ private:
 
 	bool IsSectionLocked() const;
 
+	bool CanDeleteSection() const;
+
 	void DeleteSection();
 
 	void BringToFront();
diff --git a/Engine/Source/Editor/Sequencer/Private/SequencerHotspots.cpp b/Engine/Source/Editor/Sequencer/Private/SequencerHotspots.cpp
--- a/Engine/Source/Editor/Sequencer/Private/SequencerHotspots.cpp
+++ b/Engine/Source/Editor/Sequencer/Private/SequencerHotspots.cpp
@@ -512,6 +512,11 @@ void FSectionHotspotBase::UpdateOnHover(FTrackAreaViewModel& InTrackArea) const
 
 	// Move sections if they are selected
 	TSharedPtr<ISequencer> Sequencer = WeakSequencer.Pin();
+	if (!Sequencer)
+	{
+		return;
+	}
+
 	if (Sequencer->GetViewModel()->GetSelection()->TrackArea.IsSelected(WeakSectionModel))
 	{
 		if (ThisSection && ThisSection->GetRange() != TRange<FFrameNumber>::All())
@@ -521,15 +526,23 @@ void FSectionHotspotBase::UpdateOnHover(FTrackAreaViewModel& InTrackArea) const
 	}
 	else if (ThisSection)
 	{
-		// Activate selection mode if the section has keys
-		for (const FMovieSceneChannelEntry& Entry : ThisSection->GetChannelProxy().GetAllEntries())
+		// Check if the section interface wants to show key areas (and thus supports key selection)
+		TSharedPtr<FSectionModel> SectionModel = WeakSectionModel.Pin();
+		TSharedPtr<ISequencerSection> SectionInterface = SectionModel ? SectionModel->GetSectionInterface() : nullptr;
+		const bool bShowKeyArea = !SectionInterface || SectionInterface->ShouldShowKeyAreaWhenCollapsed();
+
+		// Activate selection mode if the section has keys and wants to show them
+		if (bShowKeyArea)
 		{
-			for (const FMovieSceneChannel* Channel : Entry.GetChannels())
+			for (const FMovieSceneChannelEntry& Entry : ThisSection->GetChannelProxy().GetAllEntries())
 			{
-				if (Channel && Channel->GetNumKeys() != 0)
+				for (const FMovieSceneChannel* Channel : Entry.GetChannels())
 				{
-					InTrackArea.AttemptToActivateTool(FSequencerEditTool_Selection::Identifier);
-					return;
+					if (Channel && Channel->GetNumKeys() != 0)
+					{
+						InTrackArea.AttemptToActivateTool(FSequencerEditTool_Selection::Identifier);
+						return;
+					}
 				}
 			}
 		}
diff --git a/Engine/Source/Editor/Sequencer/Public/ISequencerSection.h b/Engine/Source/Editor/Sequencer/Public/ISequencerSection.h
--- a/Engine/Source/Editor/Sequencer/Public/ISequencerSection.h
+++ b/Engine/Source/Editor/Sequencer/Public/ISequencerSection.h
@@ -252,6 +252,19 @@ public:
 	 */
 	virtual bool IsReadOnly() const {return false;}
 
+	/**
+	 * @return Whether this section can be deleted by the user.
+	 */
+	virtual bool IsDeletable() const { return true; }
+
+	/**
+	 * Controls whether the section's key area should be visible when the section's parent
+	 * outliner item is collapsed. Default is true.
+	 * Override to return false for sections that draw their own visualization but don't want
+	 * standard key rendering when collapsed.
+	 */
+	virtual bool ShouldShowKeyAreaWhenCollapsed() const { return true; }
+
 	/**
 	 * Ticks the section during the Slate tick
 	 *
diff --git a/Engine/Source/Editor/Sequencer/Public/MVVM/ViewModels/SectionOutlinerModel.h b/Engine/Source/Editor/Sequencer/Public/MVVM/ViewModels/SectionOutlinerModel.h
--- a/Engine/Source/Editor/Sequencer/Public/MVVM/ViewModels/SectionOutlinerModel.h
+++ b/Engine/Source/Editor/Sequencer/Public/MVVM/ViewModels/SectionOutlinerModel.h
@@ -126,7 +126,7 @@ public:
 
 private:
 
-	FViewModelListHead TopLevelChannelList{EViewModelListType::TrackArea};
+	FViewModelListHead TopLevelChannelList;
 	FViewModelListHead SectionModelList{EViewModelListType::TrackArea};
 	TWeakObjectPtr<UMovieSceneSection> WeakSection;
 	TWeakViewModelPtr<FSectionModel> WeakSectionModel;
diff --git a/Engine/Source/Runtime/MovieSceneTracks/Private/Tracks/MovieSceneCommonAnimationTrack.cpp b/Engine/Source/Runtime/MovieSceneTracks/Private/Tracks/MovieSceneCommonAnimationTrack.cpp
--- a/Engine/Source/Runtime/MovieSceneTracks/Private/Tracks/MovieSceneCommonAnimationTrack.cpp
+++ b/Engine/Source/Runtime/MovieSceneTracks/Private/Tracks/MovieSceneCommonAnimationTrack.cpp
@@ -66,18 +66,21 @@ bool UMovieSceneCommonAnimationTrack::HasSection(const UMovieSceneSection& Secti
 
 void UMovieSceneCommonAnimationTrack::AddSection(UMovieSceneSection& Section)
 {
+	Modify();
 	AnimationSections.Add(&Section);
 	SetUpRootMotions(true);
 }
 
 void UMovieSceneCommonAnimationTrack::RemoveSection(UMovieSceneSection& Section)
 {
+	Modify();
 	AnimationSections.Remove(&Section);
 	SetUpRootMotions(true);
 }
 
 void UMovieSceneCommonAnimationTrack::RemoveSectionAt(int32 SectionIndex)
 {
+	Modify();
 	AnimationSections.RemoveAt(SectionIndex);
 	SetUpRootMotions(true);
 }
