Remove reference to developer module Setting when not building developer tools

#rnx

#Synced-CL 50554652

#Synced-CL 50554652

diff --git a/Engine/Plugins/Experimental/MovieSceneAnimMixer/Source/MovieSceneAnimMixer/MovieSceneAnimMixer.Build.cs b/Engine/Plugins/Experimental/MovieSceneAnimMixer/Source/MovieSceneAnimMixer/MovieSceneAnimMixer.Build.cs
--- a/Engine/Plugins/Experimental/MovieSceneAnimMixer/Source/MovieSceneAnimMixer/MovieSceneAnimMixer.Build.cs
+++ b/Engine/Plugins/Experimental/MovieSceneAnimMixer/Source/MovieSceneAnimMixer/MovieSceneAnimMixer.Build.cs
@@ -13,7 +13,6 @@ namespace UnrealBuildTool.Rules
 					"CoreUObject",
 					"Engine",
 					"DeveloperSettings",
-					"Settings",
 					"MovieScene",
 					"MovieSceneTracks",
 					"AnimGraphRuntime",
@@ -31,6 +30,11 @@ namespace UnrealBuildTool.Rules
 				}
 			);
 
+			if (Target.bBuildDeveloperTools)
+			{
+				PrivateDependencyModuleNames.Add("Settings");
+			}
+
 			if (Target.bBuildWithEditorOnlyData && Target.bBuildEditor)
 			{
 				PrivateDependencyModuleNames.AddRange(new string[]
diff --git a/Engine/Plugins/Experimental/MovieSceneAnimMixer/Source/MovieSceneAnimMixer/Private/MovieSceneAnimMixerModule.cpp b/Engine/Plugins/Experimental/MovieSceneAnimMixer/Source/MovieSceneAnimMixer/Private/MovieSceneAnimMixerModule.cpp
--- a/Engine/Plugins/Experimental/MovieSceneAnimMixer/Source/MovieSceneAnimMixer/Private/MovieSceneAnimMixerModule.cpp
+++ b/Engine/Plugins/Experimental/MovieSceneAnimMixer/Source/MovieSceneAnimMixer/Private/MovieSceneAnimMixerModule.cpp
@@ -4,10 +4,13 @@
 
 #include "AnimNode_SequencerMixerTarget.h"
 #include "MovieSceneAnimMixerSettings.h"
-#include "ISettingsModule.h"
 #include "Modules/ModuleInterface.h"
 #include "Modules/ModuleManager.h"
 
+#if WITH_UNREAL_DEVELOPER_TOOLS
+#include "ISettingsModule.h"
+#endif
+
 DEFINE_LOG_CATEGORY(LogMovieSceneAnimMixer);
 
 #define LOCTEXT_NAMESPACE "FAnimMixerModule"
@@ -25,6 +28,7 @@ namespace UE::MovieScene
 
 void FMovieSceneAnimMixerModule::StartupModule()
 {
+#if WITH_UNREAL_DEVELOPER_TOOLS
 	ISettingsModule& SettingsModule = FModuleManager::Get().LoadModuleChecked<ISettingsModule>("Settings");
 	
 	SettingsModule.RegisterSettings("Plugins", "Animation", "AnimMixer",
@@ -32,6 +36,7 @@ void FMovieSceneAnimMixerModule::StartupModule()
 		LOCTEXT("RuntimeSettingsDescription", "Configure project settings relating to the Anim Mixer Plugin"),
 		GetMutableDefault<UMovieSceneAnimMixerSettings>()
 	);
+#endif
 
 	CVarDefaultAnimNextInjectionSite->Set(*GetMutableDefault<UMovieSceneAnimMixerSettings>()->DefaultInjectionSite.ToString(), ECVF_SetByProjectSetting);
 	
diff --git a/Engine/Plugins/Experimental/MovieSceneMixedControlRig/Source/MovieSceneMixedControlRig/MovieSceneMixedControlRig.build.cs b/Engine/Plugins/Experimental/MovieSceneMixedControlRig/Source/MovieSceneMixedControlRig/MovieSceneMixedControlRig.build.cs
--- a/Engine/Plugins/Experimental/MovieSceneMixedControlRig/Source/MovieSceneMixedControlRig/MovieSceneMixedControlRig.build.cs
+++ b/Engine/Plugins/Experimental/MovieSceneMixedControlRig/Source/MovieSceneMixedControlRig/MovieSceneMixedControlRig.build.cs
@@ -13,7 +13,6 @@ namespace UnrealBuildTool.Rules
 					"CoreUObject",
 					"Engine",
 					"DeveloperSettings",
-					"Settings",
 					"MovieScene",
 					"MovieSceneTracks",
 					"AnimGraphRuntime",
