Sequencer: Anim Mixer Root Motion upgrades and Bone Matching Support

This change list updates and upgrades the Root Motion support for the Animation Mixer.
In particular:
- The Root Motion Target decoration can be added to the Animation Mixer track and allows the user to choose the target of the mixed Root Motion.
- Root Motion Settings decorations can be added to individual animation sections and support multiple modes for how to calculate the Root Motion across different spaces. These include new modes:
- Accumulated Offset, which uses previous sections' accumulated Root Motion as an offset for the Root Motion for that section automatically.
- Bone Match option, which uses the new Bone Matching workflow described further below.
- Individual sections producing Root Motion can also choose to have their Root Motion be Keep State, which means both that it can be used by the Accumulated Offset Mode as well as be used in gaps between sections that produce Root Motion. Allowing, for example, gameplay animations that are in those gaps to take advantage of Root Motion that the sequence has produced, whereas previously this would not be possible. The option to have the Root Motion apply in gaps is determined by a setting on the Root Motion Target Decoration.
- Bone matching is now accomplished through a dialog that comes up when choosing to bone match by right-clicking on a section that produces Root Motion.
- Bone matching allows the user to select a reference section, if desired, for choosing the point in time to bone match.
- Bone matching now allows a root motion offset to be computed by aligning bones between the section in question and the underlying mixed animation at a particular time, rather than just a single other animation. This will allow for smoother transitions.
- Bone matching now produces a key showing the location of where the match occurred and tracks the validity of the match as well.
- Bone matching also supports auto rematching now when moving sections or keys.

Infrastructure Changes:
- A bake system was built that bakes the result of the mixer at requested times. This is used to create an accumulated offset cache for the sections for accumulated offset mode and also can be extended in the future to do more complete bakes for more dynamic animation sections that require calculation based on time deltas, such as physics or motion matching.
- Added ISequencerDecorationEditor modular feature interface, allowing decoration objects to contribute to section context menus, provide custom key editors, custom keyframe
widgets, and receive notifications when section models are created.
- Added IMovieSceneChannelDecoration::InvalidateDecorationChannelProxy() to allow decorations to trigger dynamic channel proxy reconstruction when their internal state changes
which channels should be visible.
- Added PreFlush event to UMovieSceneEntitySystemLinker, broadcast before each ECS flush. Used by the mixer system to rebuild caches that require re-entrant ECS evaluation.
- Added SectionModelStorageExtension::NotifyDecorationEditors() to notify decoration editors when section models are created, enabling per-section subscriptions (e.g., bone
match key event handlers).
- Extended SequencerHotspots to route section context menus through registered decoration editors.
- Extended ChannelModel and IKeyArea to support custom key editor and keyframe widgets from decoration editors. This is used to have the key button for the bone matching track bring up the dialog.
- Added ShouldShowLegacyBoneMatching() virtual to FCommonAnimationSection so the AnimMixer can suppress the legacy bone matching UI in favor of its own system.
- Extracted rotation decomposition utilities (QuatFromEuler, EulerFromQuat) from MovieSceneCommonAnimationTrack.cpp into the shared header MovieSceneAnimationRotationUtils.h for
use by the bone matching system.
- Added Sequencer.cpp notification of OnSectionMoved after key snapping updates section bounds, ensuring dependent state (transitions, bone match key positions) stays in sync. This could have caused bugs elsewhere and was causing issues with keeping my state in sync for transitions and offsets.

#Synced-CL 51336034

diff --git a/Engine/Plugins/Experimental/MovieSceneAnimMixer/Source/MovieSceneAnimMixer/Private/AnimMixerBakeEvaluation.cpp b/Engine/Plugins/Experimental/MovieSceneAnimMixer/Source/MovieSceneAnimMixer/Private/AnimMixerBakeEvaluation.cpp
new file mode 100644
--- /dev/null
+++ b/Engine/Plugins/Experimental/MovieSceneAnimMixer/Source/MovieSceneAnimMixer/Private/AnimMixerBakeEvaluation.cpp
@@ -0,0 +1,197 @@
+// Copyright Epic Games, Inc. All Rights Reserved.
+
+#include "AnimMixerBakeEvaluation.h"
+
+#include "EntitySystem/MovieSceneEntitySystemLinker.h"
+#include "EntitySystem/MovieSceneEntitySystemRunner.h"
+#include "EntitySystem/MovieSceneInstanceRegistry.h"
+#include "EntitySystem/MovieSceneSequenceInstance.h"
+#include "Evaluation/MovieScenePlayback.h"
+#include "MovieSceneAnimationMixerTrack.h"
+#include "Systems/MovieSceneMixerBakeTargetSystem.h"
+#include "StructUtils/InstancedStruct.h"
+
+namespace UE::MovieScene::AnimMixerBakeEvaluation
+{
+
+// Scope that saves the mixer track's target and bake system state on construction,
+// and restores them + re-evaluates at the original time on destruction.
+//
+// The bake filter flows through the entity system via FMovieSceneMixerBakeTarget,
+// so the mixer system applies it during evaluation without a separate lookup.
+struct FBakeEvaluationScope
+{
+	UMovieSceneEntitySystemLinker* Linker;
+	FInstanceHandle InstanceHandle;
+	UMovieSceneAnimationMixerTrack* MixerTrack;
+	UMovieSceneMixerBakeTargetSystem* BakeSystem;
+	FMovieSceneEntitySystemRunner* Runner;
+	FMovieSceneContext OriginalContext;
+	TInstancedStruct<FMovieSceneMixedAnimationTarget> OriginalTarget;
+	TSharedPtr<FTransform> CaptureTarget;
+
+	FBakeEvaluationScope(
+		UMovieSceneEntitySystemLinker* InLinker,
+		FInstanceHandle InInstanceHandle,
+		UMovieSceneAnimationMixerTrack* InMixerTrack,
+		const FBakeFilter& Filter)
+		: Linker(InLinker)
+		, InstanceHandle(InInstanceHandle)
+		, MixerTrack(InMixerTrack)
+	{
+		Runner = &Linker->GetRunner().Get();
+		check(Runner);
+
+		// Save original context
+		FInstanceRegistry* Registry = Linker->GetInstanceRegistry();
+		OriginalContext = Registry->GetContext(InstanceHandle);
+
+		// Save original target, then swap to a bake target carrying the filter.
+		// The filter fields on FMovieSceneMixerBakeTarget flow through entity
+		// import into the mixer system's evaluation.
+		OriginalTarget = MixerTrack->MixedAnimationTarget;
+
+		FMovieSceneMixerBakeTarget BakeTarget;
+		BakeTarget.Filter = Filter;
+
+		// When the bake includes root motion conversion, capture the
+		// animation-space root alongside the world-space result.
+		if (!Filter.bSkipRootMotionConversion)
+		{
+			CaptureTarget = MakeShared<FTransform>(FTransform::Identity);
+			BakeTarget.Filter.CaptureAnimSpaceRoot = CaptureTarget;
+		}
+
+		MixerTrack->MixedAnimationTarget =
+			TInstancedStruct<FMovieSceneMixedAnimationTarget>::Make<FMovieSceneMixerBakeTarget>(MoveTemp(BakeTarget));
+
+		// Force full entity re-import so existing entities pick up the bake target
+		InvalidateEntities();
+
+		// Activate the bake target system
+		BakeSystem = Linker->FindSystem<UMovieSceneMixerBakeTargetSystem>();
+		if (!BakeSystem)
+		{
+			BakeSystem = Linker->LinkSystem<UMovieSceneMixerBakeTargetSystem>();
+		}
+		check(BakeSystem);
+		BakeSystem->bBakeActive = true;
+		BakeSystem->BakedResults.Reset();
+	}
+
+	~FBakeEvaluationScope()
+	{
+		// Deactivate bake
+		BakeSystem->bBakeActive = false;
+
+		// Restore original target
+		MixerTrack->MixedAnimationTarget = OriginalTarget;
+
+		// Force full entity re-import so entities pick up the restored original target
+		InvalidateEntities();
+
+		// Re-evaluate at the original time to restore current-frame state.
+		// This uses the full flush (including Finalization) since we need the
+		// original target systems to re-apply to the mesh.
+		Runner->QueueUpdate(OriginalContext, InstanceHandle);
+		Runner->Flush();
+	}
+
+	// Invalidate the sequence instance's entity cache so the next Import
+	// phase forces a full re-import of all entities with current component values.
+	void InvalidateEntities()
+	{
+		FInstanceRegistry* Registry = Linker->GetInstanceRegistry();
+		if (Registry->IsHandleValid(InstanceHandle))
+		{
+			FSequenceInstance& Instance = Registry->MutateInstance(InstanceHandle);
+			Instance.InvalidateCachedData();
+		}
+	}
+
+	// Evaluate at a single time, capturing the bake result.
+	// Sets bSilent + bHasJumped to suppress event triggers and side effects.
+	FBakeResult FlushAtTime(FFrameTime Time)
+	{
+		FFrameRate FrameRate = OriginalContext.GetFrameRate();
+		FMovieSceneEvaluationRange Range(Time, FrameRate);
+		FMovieSceneContext BakeContext(Range, EMovieScenePlayerStatus::Jumping);
+		BakeContext.SetHasJumped(true);
+		BakeContext.SetIsSilent(true);
+
+		int32 ResultIndexBefore = BakeSystem->BakedResults.Num();
+
+		Runner->QueueUpdate(BakeContext, InstanceHandle);
+		Runner->Flush();
+
+		FBakeResult Result;
+		Result.Time = Time;
+		if (BakeSystem->BakedResults.Num() > ResultIndexBefore)
+		{
+			UMovieSceneMixerBakeTargetSystem::FBakedResult& Baked = BakeSystem->BakedResults.Last();
+			Result.RootMotionTransform = Baked.RootMotionTransform;
+			Result.Keyframe = MoveTemp(Baked.Keyframe);
+		}
+
+		// When capture is active, the conversion task wrote the pre-conversion
+		// value. Otherwise (bSkipRootMotionConversion), the result is already
+		// in animation space.
+		if (CaptureTarget)
+		{
+			Result.AnimationSpaceRootMotion = *CaptureTarget;
+		}
+		else
+		{
+			Result.AnimationSpaceRootMotion = Result.RootMotionTransform;
+		}
+
+		return Result;
+	}
+};
+
+FBakeResult EvaluateAtTime(
+	UMovieSceneEntitySystemLinker* Linker,
+	FInstanceHandle InstanceHandle,
+	UMovieSceneAnimationMixerTrack* MixerTrack,
+	FFrameTime Time,
+	const FBakeFilter& Filter)
+{
+	if (!Linker || !MixerTrack)
+	{
+		return FBakeResult{ Time };
+	}
+
+	FBakeEvaluationScope Scope(Linker, InstanceHandle, MixerTrack, Filter);
+	return Scope.FlushAtTime(Time);
+}
+
+TArray<FBakeResult> EvaluateRange(
+	UMovieSceneEntitySystemLinker* Linker,
+	FInstanceHandle InstanceHandle,
+	UMovieSceneAnimationMixerTrack* MixerTrack,
+	FFrameTime StartTime,
+	FFrameTime FrameStep,
+	int32 NumSamples,
+	const FBakeFilter& Filter)
+{
+	TArray<FBakeResult> Results;
+	if (!Linker || !MixerTrack || NumSamples <= 0)
+	{
+		return Results;
+	}
+
+	Results.Reserve(NumSamples);
+
+	FBakeEvaluationScope Scope(Linker, InstanceHandle, MixerTrack, Filter);
+
+	FFrameTime CurrentTime = StartTime;
+	for (int32 i = 0; i < NumSamples; ++i)
+	{
+		Results.Add(Scope.FlushAtTime(CurrentTime));
+		CurrentTime += FrameStep;
+	}
+
+	return Results;
+}
+
+} // namespace UE::MovieScene::AnimMixerBakeEvaluation
diff --git a/Engine/Plugins/Experimental/MovieSceneAnimMixer/Source/MovieSceneAnimMixer/Private/AnimMixerBoneMatching.cpp b/Engine/Plugins/Experimental/MovieSceneAnimMixer/Source/MovieSceneAnimMixer/Private/AnimMixerBoneMatching.cpp
new file mode 100644
--- /dev/null
+++ b/Engine/Plugins/Experimental/MovieSceneAnimMixer/Source/MovieSceneAnimMixer/Private/AnimMixerBoneMatching.cpp
@@ -0,0 +1,383 @@
+// Copyright Epic Games, Inc. All Rights Reserved.
+
+#include "AnimMixerBoneMatching.h"
+
+#include "AnimMixerBakeEvaluation.h"
+#include "BoneIndices.h"
+#include "EvaluationVM/KeyframeState.h"
+#include "MovieSceneAnimationMixerTrack.h"
+#include "MovieSceneAnimationRotationUtils.h"
+#include "MovieSceneAnimTransitionSectionBase.h"
+#include "MovieSceneRootMotionSection.h"
+#include "MovieSceneRootMotionTargetDecoration.h"
+#include "MovieSceneSection.h"
+
+namespace UE::MovieScene::AnimMixerBoneMatching
+{
+
+// Compute the match offset to apply to the section's root motion so that the
+// matched bone aligns with the underlying pose.
+//
+// Both rotation and translation matching use the matched bone's component-space
+// transform. The rotation checkboxes control which Euler axes of the bone's
+// CS rotation are matched; the location checkboxes control which position axes
+// are matched. The resulting delta is applied to the root, which cascades
+// through the skeleton to move the bone into alignment.
+static FTransform ComputeMatchOffset(
+	const FTransform& UnderlyingBoneCS,
+	const FTransform& SectionBoneCS,
+	const FTransform& SectionRootCS,
+	const FTransform& SectionBoneRelToRoot,
+	const FMovieSceneBoneMatchData& Settings)
+{
+	// Selective rotation offset from the matched bone's CS rotation.
+	// Applying DeltaRot to the root rotates the entire skeleton, so the bone's CS rotation changes by the same delta.
+	const bool bAnyRotation = Settings.bMatchRotationX || Settings.bMatchRotationY || Settings.bMatchRotationZ;
+	FQuat DeltaRot = FQuat::Identity;
+
+	if (bAnyRotation)
+	{
+		ERotationOrder RotationOrder = FindBestRotationOrder(
+			Settings.bMatchRotationX, Settings.bMatchRotationY, Settings.bMatchRotationZ);
+
+		FQuat UnderlyingBoneQuat = UnderlyingBoneCS.GetRotation();
+		FQuat SecBoneQuat = SectionBoneCS.GetRotation();
+		UnderlyingBoneQuat.EnforceShortestArcWith(SecBoneQuat);
+
+		FRotator UnderlyingBoneRot(FRotator::MakeFromEuler(EulerFromQuat(UnderlyingBoneQuat, RotationOrder)));
+		FRotator SecBoneRot(FRotator::MakeFromEuler(EulerFromQuat(SecBoneQuat, RotationOrder)));
+		SecBoneRot.SetClosestToMe(UnderlyingBoneRot);
+
+		if (!Settings.bMatchRotationZ)
+		{
+			UnderlyingBoneRot.Yaw = SecBoneRot.Yaw;
+		}
+		if (!Settings.bMatchRotationY)
+		{
+			UnderlyingBoneRot.Pitch = SecBoneRot.Pitch;
+		}
+		if (!Settings.bMatchRotationX)
+		{
+			UnderlyingBoneRot.Roll = SecBoneRot.Roll;
+		}
+
+		FQuat AdjustedUnderlyingBoneQuat = QuatFromEuler(UnderlyingBoneRot.Euler(), RotationOrder);
+		DeltaRot = SecBoneQuat.Inverse() * AdjustedUnderlyingBoneQuat;
+
+		// The Euler round-trip can produce -Q instead of Q (quaternion double
+		// cover). Both represent the same rotation, but negative W causes
+		// FTransform to interpret it as a 180° flip. Enforce positive W.
+		if (DeltaRot.W < 0.0)
+		{
+			DeltaRot = -DeltaRot;
+		}
+	}
+
+	// Predict where the bone would be after applying just the rotation offset.
+	// This accounts for how rotating the root moves attached bones.
+	FQuat NewRootRot = SectionRootCS.GetRotation() * DeltaRot;
+	FTransform RotatedRoot(NewRootRot, SectionRootCS.GetTranslation());
+	RotatedRoot.SetScale3D(FVector::OneVector);
+	FTransform PredictedBone = SectionBoneRelToRoot * RotatedRoot;
+
+	// Translation delta from bone positions: how much to shift the root
+	// so the bone reaches the underlying pose's position after the rotation.
+	FVector TransDelta = UnderlyingBoneCS.GetTranslation() - PredictedBone.GetTranslation();
+
+	if (!Settings.bMatchLocationX)
+	{
+		TransDelta.X = 0.0;
+	}
+	if (!Settings.bMatchLocationY)
+	{
+		TransDelta.Y = 0.0;
+	}
+	if (!Settings.bMatchLocationZ)
+	{
+		TransDelta.Z = 0.0;
+	}
+
+	// Build the desired new root and compute the offset from the section root.
+	FTransform NewRoot(NewRootRot, SectionRootCS.GetTranslation() + TransDelta);
+	NewRoot.SetScale3D(FVector::OneVector);
+
+	return SectionRootCS.GetRelativeTransformReverse(NewRoot);
+}
+
+// Resolve the match time based on the MatchTimeMode and section boundaries.
+static FFrameTime ResolveMatchTime(
+	UMovieSceneSection* TargetSection,
+	const FMovieSceneBoneMatchData& Settings,
+	const FBoneMatchingContext& Context)
+{
+	switch (Settings.MatchTimeMode)
+	{
+	case EBoneMatchTimeMode::AtStartOfSelectedSection:
+		if (TargetSection->HasStartFrame())
+		{
+			return FFrameTime(TargetSection->GetInclusiveStartFrame());
+		}
+		break;
+
+	case EBoneMatchTimeMode::AtEndOfSelectedSection:
+		if (TargetSection->HasEndFrame())
+		{
+			return FFrameTime(TargetSection->GetExclusiveEndFrame() - 1);
+		}
+		break;
+
+	case EBoneMatchTimeMode::AtStartOfReferenceSection:
+		if (UMovieSceneSection* Ref = Settings.ReferenceSection.Get())
+		{
+			if (Ref->HasStartFrame())
+			{
+				return FFrameTime(Ref->GetInclusiveStartFrame());
+			}
+		}
+		break;
+
+	case EBoneMatchTimeMode::AtEndOfReferenceSection:
+		if (UMovieSceneSection* Ref = Settings.ReferenceSection.Get())
+		{
+			if (Ref->HasEndFrame())
+			{
+				return FFrameTime(Ref->GetExclusiveEndFrame() - 1);
+			}
+		}
+		break;
+
+	case EBoneMatchTimeMode::InBetween:
+		if (UMovieSceneSection* Ref = Settings.ReferenceSection.Get())
+		{
+			// Use the midpoint between the target's start and the reference's end.
+			FFrameNumber TargetStart = TargetSection->HasStartFrame() ? TargetSection->GetInclusiveStartFrame() : FFrameNumber(0);
+			FFrameNumber RefEnd = Ref->HasEndFrame() ? Ref->GetExclusiveEndFrame() : TargetStart;
+
+			return FFrameTime((TargetStart.Value + RefEnd.Value) / 2);
+		}
+		break;
+
+	case EBoneMatchTimeMode::AtCurrentTime:
+		return Context.CurrentTime;
+
+	default:
+		break;
+	}
+
+	// Fallback: use start of target section
+	if (TargetSection->HasStartFrame())
+	{
+		return FFrameTime(TargetSection->GetInclusiveStartFrame());
+	}
+	return FFrameTime(0);
+}
+
+FTransform ExtractBoneComponentSpaceTransform(
+	const UE::UAF::FKeyframeState* Keyframe,
+	FName BoneName)
+{
+	if (!Keyframe || !Keyframe->IsValid() || BoneName.IsNone())
+	{
+		return FTransform::Identity;
+	}
+
+	FBoneIndexType BoneIndex = Keyframe->Pose.FindLODBoneIndexFromBoneName(BoneName);
+	if (BoneIndex == INVALID_BONE_INDEX)
+	{
+		return FTransform::Identity;
+	}
+
+	// Walk up the parent chain, accumulating local transforms into component space.
+	// CS(bone) = Local(bone) * Local(parent) * Local(grandparent) * ... * Local(root)
+	FTransform ComponentSpace = (FTransform)Keyframe->Pose.LocalTransformsView[BoneIndex];
+
+	FBoneIndexType ParentIndex = Keyframe->Pose.GetLODBoneParentIndex(BoneIndex);
+	while (ParentIndex != INVALID_BONE_INDEX)
+	{
+		FTransform ParentLocal = (FTransform)Keyframe->Pose.LocalTransformsView[ParentIndex];
+		ComponentSpace = ComponentSpace * ParentLocal;
+		ParentIndex = Keyframe->Pose.GetLODBoneParentIndex(ParentIndex);
+	}
+
+	return ComponentSpace;
+}
+
+FMovieSceneBoneMatchData ComputeBoneMatch(
+	UMovieSceneSection* TargetSection,
+	const FMovieSceneBoneMatchData& Settings,
+	const FBoneMatchingContext& Context)
+{
+	FMovieSceneBoneMatchData Result = Settings;
+	Result.bIsValid = false;
+	Result.bIsDirty = false;
+	Result.MatchTransform = FTransform::Identity;
+
+	if (!TargetSection || !Context.Linker || !Context.MixerTrack || Settings.BoneName.IsNone())
+	{
+		return Result;
+	}
+
+	FFrameTime MatchTime = ResolveMatchTime(TargetSection, Settings, Context);
+
+	// Both evaluations skip root motion conversion so that RootMotionTransform
+	// stays in animation space. The bone match offset is applied in animation
+	// space (before world conversion) in the root motion conversion task.
+
+	// Underlying pose: evaluate everything except the target section,
+	// capped at the target's priority (row index).
+	AnimMixerBakeEvaluation::FBakeFilter UnderlyingFilter;
+	UnderlyingFilter.ExcludeSections.Add(FObjectKey(TargetSection));
+	UnderlyingFilter.MaxPriority = TargetSection->GetRowIndex();
+	UnderlyingFilter.bSkipRootMotionConversion = true;
+
+	AnimMixerBakeEvaluation::FBakeResult UnderlyingResult =
+		AnimMixerBakeEvaluation::EvaluateAtTime(
+			Context.Linker, Context.InstanceHandle, Context.MixerTrack,
+			MatchTime, UnderlyingFilter);
+
+	// Root motion extraction zeroes the root bone in the pose and stores it
+	// as RootMotionTransform. ExtractBoneComponentSpaceTransform walks from
+	// the bone to root (which is Identity), giving the bone's transform
+	// relative to root. Composing with RootMotionTransform gives full Component Space.
+	FTransform UnderlyingBoneRelToRoot = ExtractBoneComponentSpaceTransform(UnderlyingResult.Keyframe.Get(), Settings.BoneName);
+	UnderlyingBoneRelToRoot.NormalizeRotation();
+	FTransform UnderlyingRootCS = UnderlyingResult.RootMotionTransform;
+	UnderlyingRootCS.NormalizeRotation();
+
+	// If the reference section has its own bone match, compose its
+	// MatchTransform onto the underlying root so this section's match
+	// sees the reference's actual matched position. If the reference
+	// uses AccumulatedOffset, rebase the underlying root onto the
+	// animation-space accumulated position instead.
+	if (UMovieSceneSection* RefSection = Settings.ReferenceSection.Get())
+	{
+		if (UMovieSceneRootMotionSettingsDecoration* RefSettings = RefSection->FindDecoration<UMovieSceneRootMotionSettingsDecoration>())
+		{
+			if (RefSettings->HasBoneMatch())
+			{
+				FMovieSceneBoneMatchData RefMatch = RefSettings->GetBoneMatchData();
+				if (RefMatch.bIsValid)
+				{
+					UnderlyingRootCS = UnderlyingRootCS * RefMatch.MatchTransform;
+					UnderlyingRootCS.NormalizeRotation();
+				}
+			}
+			else if (RefSettings->GetRootTransformMode() == EMovieSceneRootMotionTransformMode::AccumulatedOffset
+				&& RefSettings->bInitialRootTransformValid
+				&& RefSection->HasStartFrame())
+			{
+				UMovieSceneRootMotionTargetDecoration* TargetDecoration = Context.MixerTrack->FindDecoration<UMovieSceneRootMotionTargetDecoration>();
+				if (TargetDecoration)
+				{
+					UMovieSceneRootMotionSection* RootMotionSection = TargetDecoration->GetRootMotionSection();
+					if (RootMotionSection)
+					{
+						FTransform AnimSpaceOffset = RootMotionSection->GetAnimSpaceAccumulatedOffsetAt(FFrameTime(RefSection->GetInclusiveStartFrame()));
+						FTransform Delta = RefSettings->InitialRootTransform.Inverse() * UnderlyingRootCS;
+						UnderlyingRootCS = Delta * AnimSpaceOffset;
+						UnderlyingRootCS.NormalizeRotation();
+					}
+				}
+			}
+		}
+	}
+
+	FTransform UnderlyingBoneCS = UnderlyingBoneRelToRoot * UnderlyingRootCS;
+
+	// Target section pose: evaluate only the target section in isolation.
+	AnimMixerBakeEvaluation::FBakeFilter SectionFilter;
+	SectionFilter.IncludeOnlySections.Add(FObjectKey(TargetSection));
+	SectionFilter.bSkipRootMotionConversion = true;
+
+	AnimMixerBakeEvaluation::FBakeResult SectionResult =
+		AnimMixerBakeEvaluation::EvaluateAtTime(
+			Context.Linker, Context.InstanceHandle, Context.MixerTrack,
+			MatchTime, SectionFilter);
+
+	FTransform SectionBoneRelToRoot = ExtractBoneComponentSpaceTransform(SectionResult.Keyframe.Get(), Settings.BoneName);
+	SectionBoneRelToRoot.NormalizeRotation();
+	FTransform SectionRootCS = SectionResult.RootMotionTransform;
+	SectionRootCS.NormalizeRotation();
+	FTransform SectionBoneCS = SectionBoneRelToRoot * SectionRootCS;
+
+	if (!UnderlyingResult.Keyframe || !SectionResult.Keyframe)
+	{
+		return Result;
+	}
+
+	Result.MatchTransform = ComputeMatchOffset(UnderlyingBoneCS, SectionBoneCS, SectionRootCS, SectionBoneRelToRoot, Settings);
+	Result.bIsValid = true;
+	Result.MatchTime = MatchTime.FloorToFrame();
+
+	return Result;
+}
+
+void PropagateRematch(
+	UMovieSceneSection* UpdatedSection,
+	const FBoneMatchingContext& Context,
+	TSet<UMovieSceneSection*>& Visited)
+{
+	UMovieSceneAnimationMixerTrack* MixerTrack = Context.MixerTrack;
+	if (!MixerTrack)
+	{
+		return;
+	}
+
+	for (UMovieSceneSection* Other : MixerTrack->GetAllSections())
+	{
+		if (!Other || Other == UpdatedSection || Visited.Contains(Other))
+		{
+			continue;
+		}
+		if (Other->IsA<UMovieSceneAnimTransitionSectionBase>())
+		{
+			continue;
+		}
+
+		UMovieSceneRootMotionSettingsDecoration* Decoration = Other->FindDecoration<UMovieSceneRootMotionSettingsDecoration>();
+		if (!Decoration || !Decoration->HasBoneMatch())
+		{
+			continue;
+		}
+
+		TMovieSceneChannelData<FMovieSceneBoneMatchData> Data = Decoration->BoneMatchChannel.GetData();
+		TArrayView<FMovieSceneBoneMatchData> Values = Data.GetValues();
+		TArrayView<const FFrameNumber> Times = Data.GetTimes();
+		if (Values.Num() == 0 || Times.Num() == 0)
+		{
+			continue;
+		}
+
+		FMovieSceneBoneMatchData& KeyValue = Values[0];
+		if (KeyValue.ReferenceSection.Get() != UpdatedSection)
+		{
+			continue;
+		}
+
+		Visited.Add(Other);
+
+		const FFrameNumber KeyTime = Times[0];
+
+		// Key must be within the overlap of the target and reference sections
+		TRange<FFrameNumber> Overlap = TRange<FFrameNumber>::Intersection(Other->GetRange(), UpdatedSection->GetRange());
+		if (Overlap.IsEmpty() || !Overlap.Contains(KeyTime))
+		{
+			KeyValue.bIsValid = false;
+			KeyValue.bIsDirty = true;
+			continue;
+		}
+
+		FBoneMatchingContext ChainContext = Context;
+		ChainContext.CurrentTime = FFrameTime(KeyTime);
+
+		FMovieSceneBoneMatchData Result = ComputeBoneMatch(Other, KeyValue, ChainContext);
+		KeyValue = Result;
+
+		if (Result.bIsValid)
+		{
+			PropagateRematch(Other, Context, Visited);
+		}
+	}
+}
+
+} // namespace UE::MovieScene::AnimMixerBoneMatching
diff --git a/Engine/Plugins/Experimental/MovieSceneAnimMixer/Source/MovieSceneAnimMixer/Private/AnimMixerComponentTypes.cpp b/Engine/Plugins/Experimental/MovieSceneAnimMixer/Source/MovieSceneAnimMixer/Private/AnimMixerComponentTypes.cpp
--- a/Engine/Plugins/Experimental/MovieSceneAnimMixer/Source/MovieSceneAnimMixer/Private/AnimMixerComponentTypes.cpp
+++ b/Engine/Plugins/Experimental/MovieSceneAnimMixer/Source/MovieSceneAnimMixer/Private/AnimMixerComponentTypes.cpp
@@ -1,6 +1,7 @@
 // Copyright Epic Games, Inc. All Rights Reserved.
 
 #include "AnimMixerComponentTypes.h"
+#include "MovieSceneAnimationMixerTrack.h"
 #include "EntitySystem/BuiltInComponentTypes.h"
 #include "EntitySystem/MovieSceneComponentRegistry.h"
 #include "EntitySystem/MovieSceneEntityFactoryTemplates.h"
@@ -49,6 +50,8 @@ namespace UE::MovieScene
 		ComponentRegistry->NewComponentType(&RootDestination, TEXT("Root Destination"));
 		ComponentRegistry->NewComponentType(&MixerLayer, TEXT("Mixer Layer"));
 		ComponentRegistry->NewComponentType(&MirrorTable, TEXT("Mirror Table"));
+		ComponentRegistry->NewComponentType(&BoneMatchTransform, TEXT("Bone Match Transform"));
+		ComponentRegistry->NewComponentType(&GapBehavior, TEXT("Gap Behavior"));
 		ComponentRegistry->NewComponentType(&EntityOwner, TEXT("Entity Owner"));
 
 
@@ -65,7 +68,9 @@ namespace UE::MovieScene
 		ComponentRegistry->Factories.DuplicateChildComponent(MixerRootMotion);
 		ComponentRegistry->Factories.DuplicateChildComponent(RootDestination);
 		ComponentRegistry->Factories.DuplicateChildComponent(MixerLayer);
+		ComponentRegistry->Factories.DuplicateChildComponent(GapBehavior);
 		ComponentRegistry->Factories.DuplicateChildComponent(EntityOwner);
+		ComponentRegistry->Factories.DuplicateChildComponent(BoneMatchTransform);
 
 		ComponentRegistry->Factories.DuplicateChildComponent(MirrorTable);
 		
diff --git a/Engine/Plugins/Experimental/MovieSceneAnimMixer/Source/MovieSceneAnimMixer/Private/MovieSceneAnimationMixerTrack.cpp b/Engine/Plugins/Experimental/MovieSceneAnimMixer/Source/MovieSceneAnimMixer/Private/MovieSceneAnimationMixerTrack.cpp
--- a/Engine/Plugins/Experimental/MovieSceneAnimMixer/Source/MovieSceneAnimMixer/Private/MovieSceneAnimationMixerTrack.cpp
+++ b/Engine/Plugins/Experimental/MovieSceneAnimMixer/Source/MovieSceneAnimMixer/Private/MovieSceneAnimationMixerTrack.cpp
@@ -3,6 +3,7 @@
 #include "MovieSceneAnimationMixerTrack.h"
 #include "MovieSceneAnimTransitionSectionBase.h"
 #include "MovieSceneAnimCrossfadeTransitionSection.h"
+#include "MovieSceneRootMotionTargetDecoration.h"
 #include "MovieSceneSection.h"
 #include "Sections/MovieSceneSkeletalAnimationSection.h"
 #include "Channels/MovieSceneChannelProxy.h"
@@ -16,6 +17,8 @@
 #include "MovieScene.h"
 #include "MovieSceneSequence.h"
 #include "MovieSceneTrack.h"
+#include "AnimMixerBakeEvaluation.h"
+#include "MovieSceneRootMotionSection.h"
 
 #define LOCTEXT_NAMESPACE "MovieSceneAnimationMixerTrack"
 
@@ -46,6 +49,20 @@ UMovieSceneRootMotionSettingsDecoration::UMovieSceneRootMotionSettingsDecoration
 
 	TransformMode.SetEnum(StaticEnum<EMovieSceneRootMotionTransformMode>());
 	TransformMode.SetDefault((uint8)EMovieSceneRootMotionTransformMode::Offset);
+
+	KeepStateChannel.SetEnum(StaticEnum<EMovieSceneRootMotionKeepState>());
+	KeepStateChannel.SetDefault((uint8)EMovieSceneRootMotionKeepState::DontKeep);
+
+	// Invalidate root motion caches whenever this decoration is modified.
+	// This covers channel default edits (which only call Modify(), not
+	// PostEditChangeProperty) as well as property edits and undo.
+	OnSignatureChanged().AddWeakLambda(this, [this]()
+	{
+		if (UMovieSceneAnimationMixerTrack* MixerTrack = GetTypedOuter<UMovieSceneAnimationMixerTrack>())
+		{
+			MixerTrack->InvalidateAccumulatedOffsetCache();
+		}
+	});
 }
 
 EMovieSceneRootMotionTransformMode UMovieSceneRootMotionSettingsDecoration::GetRootTransformMode() const
@@ -59,20 +76,28 @@ EMovieSceneChannelProxyType UMovieSceneRootMotionSettingsDecoration::PopulateCha
 
 #if WITH_EDITOR
 	FText Group = LOCTEXT("RootMotionSettings", "Root Motion Settings");
+	int32 SortOrder = 0;
 	{
 		FMovieSceneChannelMetaData ChannelMetaData;
 		ChannelMetaData.SetIdentifiers("Space", LOCTEXT("Space", "Space"), Group);
 		ChannelMetaData.WeakOwningObject = this;
-		ChannelMetaData.SortOrder = 0;
+		ChannelMetaData.SortOrder = SortOrder++;
 		OutProxyData.Add(RootMotionSpace, ChannelMetaData, TMovieSceneExternalValue<uint8>());
 	}
 	{
 		FMovieSceneChannelMetaData ChannelMetaData;
 		ChannelMetaData.SetIdentifiers("Root Mode", LOCTEXT("Mode", "Mode"), Group);
 		ChannelMetaData.WeakOwningObject = this;
-		ChannelMetaData.SortOrder = 1;
+		ChannelMetaData.SortOrder = SortOrder++;
 		OutProxyData.Add(TransformMode, ChannelMetaData, TMovieSceneExternalValue<uint8>());
 	}
+	{
+		FMovieSceneChannelMetaData ChannelMetaData;
+		ChannelMetaData.SetIdentifiers("KeepState", LOCTEXT("KeepState", "Keep State"), Group);
+		ChannelMetaData.WeakOwningObject = this;
+		ChannelMetaData.SortOrder = SortOrder++;
+		OutProxyData.Add(KeepStateChannel, ChannelMetaData, TMovieSceneExternalValue<uint8>());
+	}
 
 	if (CurrentTransformMode != EMovieSceneRootMotionTransformMode::Asset)
 	{
@@ -83,7 +108,7 @@ EMovieSceneChannelProxyType UMovieSceneRootMotionSettingsDecoration::PopulateCha
 			ChannelMetaData.SetIdentifiers("RootBaseLocation.X", AxisDisplayInfo::GetAxisDisplayName(EAxisList::X), Group);
 			ChannelMetaData.Color = AxisDisplayInfo::GetAxisColor(EAxisList::X);
 			ChannelMetaData.WeakOwningObject = this;
-			ChannelMetaData.SortOrder = 2;
+			ChannelMetaData.SortOrder = SortOrder++;
 			ChannelMetaData.PropertyMetaData.Add(SliderExponent, TEXT("0.2"));
 			OutProxyData.Add(Location[0], ChannelMetaData, TMovieSceneExternalValue<double>());
 		}
@@ -92,7 +117,7 @@ EMovieSceneChannelProxyType UMovieSceneRootMotionSettingsDecoration::PopulateCha
 			ChannelMetaData.SetIdentifiers("RootBaseLocation.Y", AxisDisplayInfo::GetAxisDisplayName(EAxisList::Y), Group);
 			ChannelMetaData.Color = AxisDisplayInfo::GetAxisColor(EAxisList::Y);
 			ChannelMetaData.WeakOwningObject = this;
-			ChannelMetaData.SortOrder = 3;
+			ChannelMetaData.SortOrder = SortOrder++;
 			ChannelMetaData.PropertyMetaData.Add(SliderExponent, TEXT("0.2"));
 			OutProxyData.Add(Location[1], ChannelMetaData, TMovieSceneExternalValue<double>());
 		}
@@ -101,7 +126,7 @@ EMovieSceneChannelProxyType UMovieSceneRootMotionSettingsDecoration::PopulateCha
 			ChannelMetaData.SetIdentifiers("RootBaseLocation.Z", AxisDisplayInfo::GetAxisDisplayName(EAxisList::Z), Group);
 			ChannelMetaData.Color = AxisDisplayInfo::GetAxisColor(EAxisList::Z);
 			ChannelMetaData.WeakOwningObject = this;
-			ChannelMetaData.SortOrder = 4;
+			ChannelMetaData.SortOrder = SortOrder++;
 			ChannelMetaData.PropertyMetaData.Add(SliderExponent, TEXT("0.2"));
 			OutProxyData.Add(Location[2], ChannelMetaData, TMovieSceneExternalValue<double>());
 		}
@@ -110,7 +135,7 @@ EMovieSceneChannelProxyType UMovieSceneRootMotionSettingsDecoration::PopulateCha
 			ChannelMetaData.SetIdentifiers("RootBaseRotation.X", LOCTEXT("RotationX", "Roll"), Group);
 			ChannelMetaData.Color = AxisDisplayInfo::GetAxisColor(EAxisList::X);
 			ChannelMetaData.WeakOwningObject = this;
-			ChannelMetaData.SortOrder = 5;
+			ChannelMetaData.SortOrder = SortOrder++;
 			ChannelMetaData.PropertyMetaData.Add(SliderExponent, TEXT("0.2"));
 			OutProxyData.Add(Rotation[0], ChannelMetaData, TMovieSceneExternalValue<double>());
 		}
@@ -119,7 +144,7 @@ EMovieSceneChannelProxyType UMovieSceneRootMotionSettingsDecoration::PopulateCha
 			ChannelMetaData.SetIdentifiers("RootBaseRotation.Y", LOCTEXT("RotationY", "Pitch"), Group);
 			ChannelMetaData.Color = AxisDisplayInfo::GetAxisColor(EAxisList::Y);
 			ChannelMetaData.WeakOwningObject = this;
-			ChannelMetaData.SortOrder = 6;
+			ChannelMetaData.SortOrder = SortOrder++;
 			ChannelMetaData.PropertyMetaData.Add(SliderExponent, TEXT("0.2"));
 			OutProxyData.Add(Rotation[1], ChannelMetaData, TMovieSceneExternalValue<double>());
 		}
@@ -128,12 +153,28 @@ EMovieSceneChannelProxyType UMovieSceneRootMotionSettingsDecoration::PopulateCha
 			ChannelMetaData.SetIdentifiers("RootBaseRotation.Z", LOCTEXT("RotationZ", "Yaw"), Group);
 			ChannelMetaData.Color = AxisDisplayInfo::GetAxisColor(EAxisList::Z);
 			ChannelMetaData.WeakOwningObject = this;
-			ChannelMetaData.SortOrder = 7;
+			ChannelMetaData.SortOrder = SortOrder++;
 			ChannelMetaData.PropertyMetaData.Add(SliderExponent, TEXT("0.2"));
 			OutProxyData.Add(Rotation[2], ChannelMetaData, TMovieSceneExternalValue<double>());
 		}
 	}
 
+#if WITH_EDITOR
+	if (CurrentTransformMode == EMovieSceneRootMotionTransformMode::BoneMatch && BoneMatchChannel.GetNumKeys() > 0)
+	{
+		const bool bMatchValid = GetBoneMatchData().bIsValid;
+		FMovieSceneChannelMetaData ChannelMetaData;
+		ChannelMetaData.SetIdentifiers("BoneMatch",
+			bMatchValid
+				? LOCTEXT("BoneMatch", "Bone Match")
+				: LOCTEXT("BoneMatchInvalid", "Bone Match (invalid)"),
+			Group);
+		ChannelMetaData.WeakOwningObject = this;
+		ChannelMetaData.SortOrder = SortOrder++;
+		OutProxyData.Add(BoneMatchChannel, ChannelMetaData);
+	}
+#endif
+
 #else
 
 	if (CurrentTransformMode != EMovieSceneRootMotionTransformMode::Asset)
@@ -142,6 +183,8 @@ EMovieSceneChannelProxyType UMovieSceneRootMotionSettingsDecoration::PopulateCha
 		OutProxyData.Add(TransformMode);
 	}
 
+	OutProxyData.Add(KeepStateChannel);
+
 	OutProxyData.Add(Location[0]);
 	OutProxyData.Add(Location[1]);
 	OutProxyData.Add(Location[2]);
@@ -150,6 +193,8 @@ EMovieSceneChannelProxyType UMovieSceneRootMotionSettingsDecoration::PopulateCha
 	OutProxyData.Add(Rotation[1]);
 	OutProxyData.Add(Rotation[2]);
 
+	OutProxyData.Add(BoneMatchChannel);
+
 #endif
 
 	return EMovieSceneChannelProxyType::Dynamic;
@@ -185,6 +230,7 @@ void UMovieSceneRootMotionSettingsDecoration::ExtendEntityImpl(UMovieSceneEntity
 		.AddConditional(BuiltInComponents->DoubleChannel[4], &Rotation[1], bAddTransform)
 		.AddConditional(BuiltInComponents->DoubleChannel[5], &Rotation[2], bAddTransform)
 		.Add(AnimMixerComponents->RootMotionSettings, RootMotionSettings)
+		.AddConditional(AnimMixerComponents->BoneMatchTransform, GetBoneMatchData().MatchTransform, HasBoneMatch() && GetBoneMatchData().bIsValid)
 	);
 }
 
@@ -217,6 +263,32 @@ void UMovieSceneAnimationSectionDecoration::ExtendEntityImpl(UMovieSceneEntitySy
 	);
 }
 
+bool WouldCreateBoneMatchCycle(UMovieSceneSection* CandidateSection, UMovieSceneSection* TargetSection)
+{
+	TSet<UMovieSceneSection*> Visited;
+	UMovieSceneSection* Current = CandidateSection;
+	while (Current)
+	{
+		if (Current == TargetSection)
+		{
+			return true;
+		}
+		if (Visited.Contains(Current))
+		{
+			break;
+		}
+		Visited.Add(Current);
+
+		UMovieSceneRootMotionSettingsDecoration* Decoration = Current->FindDecoration<UMovieSceneRootMotionSettingsDecoration>();
+		if (!Decoration || !Decoration->HasBoneMatch())
+		{
+			break;
+		}
+		Current = Decoration->GetBoneMatchData().ReferenceSection.Get();
+	}
+	return false;
+}
+
 UMovieSceneAnimationMixerTrack::UMovieSceneAnimationMixerTrack(const FObjectInitializer& ObjInit)
 	: Super(ObjInit)
 {
@@ -296,8 +368,7 @@ void UMovieSceneAnimationMixerTrack::OnSectionAddedImpl(UMovieSceneSection* Sect
 			UpdateTransitionsForRow(RowIndex);
 		}
 
-		// Mark root motions as dirty so accumulated offset cache will be rebuilt
-		SetRootMotionsDirty();
+		InvalidateAccumulatedOffsetCache();
 	}
 }
 
@@ -319,8 +390,33 @@ void UMovieSceneAnimationMixerTrack::OnSectionRemovedImpl(UMovieSceneSection* Se
 			RemoveSection(*Transition);
 		}
 
-		// Mark root motions as dirty so accumulated offset cache will be rebuilt
-		SetRootMotionsDirty();
+		// Mark bone matches that referenced the removed section as invalid
+		for (UMovieSceneSection* OtherSection : GetAllSections())
+		{
+			if (!OtherSection || OtherSection == Section || OtherSection->IsA<UMovieSceneAnimTransitionSectionBase>())
+			{
+				continue;
+			}
+
+			UMovieSceneRootMotionSettingsDecoration* OtherDecoration = OtherSection->FindDecoration<UMovieSceneRootMotionSettingsDecoration>();
+			if (!OtherDecoration || !OtherDecoration->HasBoneMatch())
+			{
+				continue;
+			}
+
+			TMovieSceneChannelData<FMovieSceneBoneMatchData> Data = OtherDecoration->BoneMatchChannel.GetData();
+			TArrayView<FMovieSceneBoneMatchData> Values = Data.GetValues();
+			for (FMovieSceneBoneMatchData& Value : Values)
+			{
+				if (Value.ReferenceSection.Get() == Section)
+				{
+					Value.bIsValid = false;
+					Value.bIsDirty = true;
+				}
+			}
+		}
+
+		InvalidateAccumulatedOffsetCache();
 	}
 }
 
@@ -737,8 +833,82 @@ EMovieSceneSectionMovedResult UMovieSceneAnimationMixerTrack::OnSectionMoved(UMo
 
 	const bool bSectionsChanged = UpdateTransitionsForRow(Section.GetRowIndex());
 
-	// Mark root motions as dirty so accumulated offset cache will be rebuilt
-	SetRootMotionsDirty();
+	// Reference section tracking: when a section moves, find bone matches
+	// on other sections that reference it. Update their key times so the
+	// bone match key follows the reference section's anchor point.
+	// This runs on every move tick (including interactive drag) so the
+	// key stays in sync during the drag.
+	for (UMovieSceneSection* OtherSection : GetAllSections())
+	{
+		if (!OtherSection || OtherSection == &Section || OtherSection->IsA<UMovieSceneAnimTransitionSectionBase>())
+		{
+			continue;
+		}
+
+		UMovieSceneRootMotionSettingsDecoration* OtherDecoration = OtherSection->FindDecoration<UMovieSceneRootMotionSettingsDecoration>();
+		if (!OtherDecoration || !OtherDecoration->HasBoneMatch())
+		{
+			continue;
+		}
+
+		FMovieSceneBoneMatchData MatchData = OtherDecoration->GetBoneMatchData();
+		if (MatchData.ReferenceSection.Get() != &Section)
+		{
+			continue;
+		}
+
+		// Resolve the new key time based on the reference section's new position
+		FFrameNumber NewTime;
+		switch (MatchData.MatchTimeMode)
+		{
+		case EBoneMatchTimeMode::AtStartOfReferenceSection:
+			NewTime = Section.HasStartFrame() ? Section.GetInclusiveStartFrame() : FFrameNumber(0);
+			break;
+
+		case EBoneMatchTimeMode::AtEndOfReferenceSection:
+			NewTime = Section.HasEndFrame() ? (Section.GetExclusiveEndFrame() - 1) : FFrameNumber(0);
+			break;
+
+		case EBoneMatchTimeMode::InBetween:
+		{
+			FFrameNumber TargetStart = OtherSection->HasStartFrame()
+				? OtherSection->GetInclusiveStartFrame() : FFrameNumber(0);
+			FFrameNumber RefEnd = Section.HasEndFrame()
+				? Section.GetExclusiveEndFrame() : TargetStart;
+			NewTime = FFrameNumber((TargetStart.Value + RefEnd.Value) / 2);
+			break;
+		}
+
+		default:
+			continue;
+		}
+
+		const bool bWasValid = MatchData.bIsValid;
+		OtherDecoration->BoneMatchChannel.NotifyReferenceChanged(NewTime);
+
+		// Dirty the decoration proxy when validity flips so the channel
+		// display name updates.
+		FMovieSceneBoneMatchData UpdatedData = OtherDecoration->GetBoneMatchData();
+		if (UpdatedData.bIsValid != bWasValid)
+		{
+			OtherDecoration->InvalidateDecorationChannelProxy();
+		}
+	}
+
+	InvalidateAccumulatedOffsetCache();
+
+	if (Params.MoveType == EPropertyChangeType::ValueSet)
+	{
+		// Dirty the moved section's decoration proxy so any bone match
+		// validity change is reflected in the channel display.
+		if (UMovieSceneRootMotionSettingsDecoration* Decoration = Section.FindDecoration<UMovieSceneRootMotionSettingsDecoration>())
+		{
+			if (Decoration->HasBoneMatch())
+			{
+				Decoration->InvalidateDecorationChannelProxy();
+			}
+		}
+	}
 
 	return bSectionsChanged ? EMovieSceneSectionMovedResult::SectionsChanged : EMovieSceneSectionMovedResult::None;
 }
@@ -1060,4 +1230,82 @@ UMovieSceneAnimationMixerLayer* UMovieSceneAnimationMixerTrack::InsertLayer(int3
 	return NewLayer;
 }
 
+void UMovieSceneAnimationMixerTrack::InvalidateAccumulatedOffsetCache()
+{
+	SetRootMotionsDirty();
+
+	if (UMovieSceneRootMotionTargetDecoration* Decoration = FindDecoration<UMovieSceneRootMotionTargetDecoration>())
+	{
+		if (UMovieSceneRootMotionSection* RootMotionSection = Decoration->GetRootMotionSection())
+		{
+			RootMotionSection->InvalidateAccumulatedOffsetCache();
+		}
+	}
+}
+
+bool UMovieSceneAnimationMixerTrack::HasDirtyAccumulatedOffsetCache() const
+{
+	if (UMovieSceneRootMotionTargetDecoration* Decoration = FindDecoration<UMovieSceneRootMotionTargetDecoration>())
+	{
+		if (const UMovieSceneRootMotionSection* RootMotionSection = Decoration->GetRootMotionSection())
+		{
+			return RootMotionSection->bAccumulatedOffsetCacheDirty;
+		}
+	}
+	return false;
+}
+
+void UMovieSceneAnimationMixerTrack::RebuildDirtyAccumulatedOffsetCache(
+	UMovieSceneEntitySystemLinker* Linker,
+	UE::MovieScene::FInstanceHandle InstanceHandle)
+{
+	if (UMovieSceneRootMotionTargetDecoration* Decoration = FindDecoration<UMovieSceneRootMotionTargetDecoration>())
+	{
+		if (UMovieSceneRootMotionSection* RootMotionSection = Decoration->GetRootMotionSection())
+		{
+			if (!RootMotionSection->bAccumulatedOffsetCacheDirty)
+			{
+				return;
+			}
+
+			RootMotionSection->bAccumulatedOffsetCacheDirty = false;
+
+			// Compute initial root transforms for AccumulatedOffset sections.
+			// Each section is baked in isolation at its start time with
+			// bSkipRootMotionConversion to get the raw animation-space root
+			// transform. At runtime, the delta is re-based onto the accumulated position.
+			for (UMovieSceneSection* Section : GetAllSections())
+			{
+				if (!Section)
+				{
+					continue;
+				}
+
+				auto* Settings = Section->FindDecoration<UMovieSceneRootMotionSettingsDecoration>();
+				if (!Settings
+					|| Settings->GetRootTransformMode() != EMovieSceneRootMotionTransformMode::AccumulatedOffset
+					|| !Section->HasStartFrame())
+				{
+					continue;
+				}
+
+				FFrameTime SectionStart(Section->GetInclusiveStartFrame());
+
+				using namespace UE::MovieScene::AnimMixerBakeEvaluation;
+				FBakeFilter Filter;
+				Filter.IncludeOnlySections.Add(FObjectKey(Section));
+				Filter.bSkipRootMotionConversion = true;
+
+				FBakeResult Result = EvaluateAtTime(
+					Linker, InstanceHandle, this, SectionStart, Filter);
+
+				Settings->InitialRootTransform = Result.RootMotionTransform;
+				Settings->bInitialRootTransformValid = true;
+			}
+
+			RootMotionSection->RebuildAccumulatedOffsetCache(Linker, InstanceHandle, this);
+		}
+	}
+}
+
 #undef LOCTEXT_NAMESPACE
\ No newline at end of file
diff --git a/Engine/Plugins/Experimental/MovieSceneAnimMixer/Source/MovieSceneAnimMixer/Private/MovieSceneBoneMatchChannel.cpp b/Engine/Plugins/Experimental/MovieSceneAnimMixer/Source/MovieSceneAnimMixer/Private/MovieSceneBoneMatchChannel.cpp
new file mode 100644
--- /dev/null
+++ b/Engine/Plugins/Experimental/MovieSceneAnimMixer/Source/MovieSceneAnimMixer/Private/MovieSceneBoneMatchChannel.cpp
@@ -0,0 +1,139 @@
+// Copyright Epic Games, Inc. All Rights Reserved.
+
+#include "MovieSceneBoneMatchChannel.h"
+
+#include UE_INLINE_GENERATED_CPP_BY_NAME(MovieSceneBoneMatchChannel)
+
+void FMovieSceneBoneMatchChannel::GetKeys(const TRange<FFrameNumber>& WithinRange, TArray<FFrameNumber>* OutKeyTimes, TArray<FKeyHandle>* OutKeyHandles)
+{
+	GetData().GetKeys(WithinRange, OutKeyTimes, OutKeyHandles);
+}
+
+void FMovieSceneBoneMatchChannel::GetKeyTimes(TArrayView<const FKeyHandle> InHandles, TArrayView<FFrameNumber> OutKeyTimes)
+{
+	GetData().GetKeyTimes(InHandles, OutKeyTimes);
+}
+
+void FMovieSceneBoneMatchChannel::SetKeyTimes(TArrayView<const FKeyHandle> InHandles, TArrayView<const FFrameNumber> InKeyTimes)
+{
+	// Manual key drag disconnects from anchors: switch to AtCurrentTime
+	// and clear the reference section.
+	for (int32 i = 0; i < InHandles.Num(); ++i)
+	{
+		const int32 KeyIndex = GetData().GetIndex(InHandles[i]);
+		if (KeyIndex != INDEX_NONE)
+		{
+			Values[KeyIndex].MatchTimeMode = EBoneMatchTimeMode::AtCurrentTime;
+			Values[KeyIndex].ReferenceSection = nullptr;
+		}
+	}
+
+	// MoveKeyInternal fires OnKeyMovedEvent automatically
+	GetData().SetKeyTimes(InHandles, InKeyTimes);
+}
+
+void FMovieSceneBoneMatchChannel::DuplicateKeys(TArrayView<const FKeyHandle> InHandles, TArrayView<FKeyHandle> OutNewHandles)
+{
+	GetData().DuplicateKeys(InHandles, OutNewHandles);
+}
+
+void FMovieSceneBoneMatchChannel::DeleteKeys(TArrayView<const FKeyHandle> InHandles)
+{
+	GetData().DeleteKeys(InHandles);
+}
+
+void FMovieSceneBoneMatchChannel::DeleteKeysFrom(FFrameNumber InTime, bool bDeleteKeysBefore)
+{
+	GetData().DeleteKeysFrom(InTime, bDeleteKeysBefore);
+}
+
+void FMovieSceneBoneMatchChannel::RemapTimes(const UE::MovieScene::IRetimingInterface& Retimer)
+{
+	GetData().RemapTimes(Retimer);
+}
+
+TRange<FFrameNumber> FMovieSceneBoneMatchChannel::ComputeEffectiveRange() const
+{
+	return GetData().GetTotalRange();
+}
+
+int32 FMovieSceneBoneMatchChannel::GetNumKeys() const
+{
+	return Times.Num();
+}
+
+void FMovieSceneBoneMatchChannel::Reset()
+{
+	GetData().Reset();
+}
+
+void FMovieSceneBoneMatchChannel::Offset(FFrameNumber DeltaPosition)
+{
+	// Section-relative modes move with the section; other modes keep
+	// their key time fixed (they are anchored to a reference section
+	// or to the playhead).
+	TArray<FKeyMoveEventItem> Items;
+	for (int32 i = 0; i < Times.Num(); ++i)
+	{
+		const FFrameNumber OldTime = Times[i];
+		const EBoneMatchTimeMode Mode = Values[i].MatchTimeMode;
+
+		if (Mode == EBoneMatchTimeMode::AtStartOfSelectedSection ||
+			Mode == EBoneMatchTimeMode::AtEndOfSelectedSection)
+		{
+			Times[i] = OldTime + DeltaPosition;
+		}
+
+		Items.Add(FKeyMoveEventItem(i, OldTime, i, Times[i]));
+	}
+
+	if (Items.Num() > 0)
+	{
+		OnKeyMovedEvent().Broadcast(this, Items);
+	}
+}
+
+void FMovieSceneBoneMatchChannel::NotifyReferenceChanged(FFrameNumber NewResolvedTime)
+{
+	// Called when a reference section moves. The caller already verified
+	// that this channel's bone match references the moved section.
+	// Only adjust keys whose mode is anchored to a reference section.
+	TArray<FKeyMoveEventItem> Items;
+	for (int32 i = 0; i < Times.Num(); ++i)
+	{
+		const EBoneMatchTimeMode Mode = Values[i].MatchTimeMode;
+		if (Mode != EBoneMatchTimeMode::AtStartOfReferenceSection &&
+			Mode != EBoneMatchTimeMode::AtEndOfReferenceSection &&
+			Mode != EBoneMatchTimeMode::InBetween)
+		{
+			continue;
+		}
+
+		const FFrameNumber OldTime = Times[i];
+		Times[i] = NewResolvedTime;
+		Items.Add(FKeyMoveEventItem(i, OldTime, i, NewResolvedTime));
+	}
+
+	if (Items.Num() > 0)
+	{
+		OnKeyMovedEvent().Broadcast(this, Items);
+	}
+}
+
+void FMovieSceneBoneMatchChannel::Optimize(const FKeyDataOptimizationParams& InParameters)
+{
+}
+
+void FMovieSceneBoneMatchChannel::ClearDefault()
+{
+}
+
+FKeyHandle FMovieSceneBoneMatchChannel::GetHandle(int32 Index)
+{
+	return GetData().GetHandle(Index);
+}
+
+int32 FMovieSceneBoneMatchChannel::GetIndex(FKeyHandle Handle)
+{
+	return GetData().GetIndex(Handle);
+}
diff --git a/Engine/Plugins/Experimental/MovieSceneAnimMixer/Source/MovieSceneAnimMixer/Private/MovieSceneRootMotionSection.cpp b/Engine/Plugins/Experimental/MovieSceneAnimMixer/Source/MovieSceneAnimMixer/Private/MovieSceneRootMotionSection.cpp
--- a/Engine/Plugins/Experimental/MovieSceneAnimMixer/Source/MovieSceneAnimMixer/Private/MovieSceneRootMotionSection.cpp
+++ b/Engine/Plugins/Experimental/MovieSceneAnimMixer/Source/MovieSceneAnimMixer/Private/MovieSceneRootMotionSection.cpp
@@ -2,10 +2,12 @@
 
 #include "MovieSceneRootMotionSection.h"
 
+#include "AnimMixerBakeEvaluation.h"
 #include "AnimMixerComponentTypes.h"
 #include "Channels/MovieSceneChannelProxy.h"
 #include "EntitySystem/BuiltInComponentTypes.h"
 #include "Evaluation/MovieSceneEvaluationField.h"
+#include "MovieSceneAnimTransitionSectionBase.h"
 #include "MovieSceneTrack.h"
 #include "Systems/MovieSceneSkeletalAnimationSystem.h"
 #include "Tracks/MovieSceneCommonAnimationTrack.h"
@@ -17,6 +19,15 @@ UMovieSceneRootMotionSection::UMovieSceneRootMotionSection(const FObjectInitiali
 	RootDestinationChannel.SetDefault((uint8)EMovieSceneRootMotionDestination::RootBone);
 
 	SectionRange.Value = TRange<FFrameNumber>::All();
+
+	// Invalidate root motion caches whenever this section is modified.
+	OnSignatureChanged().AddWeakLambda(this, [this]()
+	{
+		if (UMovieSceneAnimationMixerTrack* MixerTrack = GetTypedOuter<UMovieSceneAnimationMixerTrack>())
+		{
+			MixerTrack->InvalidateAccumulatedOffsetCache();
+		}
+	});
 }
 
 EMovieSceneChannelProxyType UMovieSceneRootMotionSection::CacheChannelProxy()
@@ -57,6 +68,8 @@ void UMovieSceneRootMotionSection::ImportEntityImpl(UMovieSceneEntitySystemLinke
 	{
 		UMovieSceneAnimationMixerTrack* AnimTrack = GetTypedOuter<UMovieSceneAnimationMixerTrack>();
 
+		const bool bPersistGaps = (GapBehavior == EMovieSceneRootMotionGapBehavior::PersistPreviousTransform);
+
 		OutImportedEntity->AddBuilder(
 			FEntityBuilder()
 			.Add(BuiltInComponents->GenericObjectBinding, Params.GetObjectBindingID())
@@ -64,13 +77,20 @@ void UMovieSceneRootMotionSection::ImportEntityImpl(UMovieSceneEntitySystemLinke
 			.Add(AnimMixerComponents->Target, AnimTrack->MixedAnimationTarget)
 			.Add(BuiltInComponents->ByteChannel, FSourceByteChannel{ &RootDestinationChannel })
 			.AddDefaulted(AnimMixerComponents->RootDestination)
+			.Add(AnimMixerComponents->GapBehavior, GapBehavior)
+			.Add(AnimMixerComponents->EntityOwner, FObjectKey(this))
+			.Add(AnimMixerComponents->Priority, MIN_int32)
+			// Add a null Task so the mixer system creates a gap sentinel entry,
+			// keeping the mixer alive during gaps for root motion persistence.
+			.AddConditional(AnimMixerComponents->Task, TSharedPtr<FAnimNextEvaluationTask>(), bPersistGaps)
 		);
 	}
 }
 
 bool UMovieSceneRootMotionSection::PopulateEvaluationFieldImpl(const TRange<FFrameNumber>& EffectiveRange, const FMovieSceneEvaluationFieldEntityMetaData& InMetaData, FMovieSceneEntityComponentFieldBuilder* OutFieldBuilder)
 {
-	if (RootDestinationChannel.HasAnyData())
+	if (RootDestinationChannel.HasAnyData()
+		|| GapBehavior == EMovieSceneRootMotionGapBehavior::PersistPreviousTransform)
 	{
 		const int32 MetaDataIndex = OutFieldBuilder->AddMetaData(InMetaData);
 		OutFieldBuilder->AddPersistentEntity(EffectiveRange, this, 0, MetaDataIndex);
@@ -81,22 +101,26 @@ bool UMovieSceneRootMotionSection::PopulateEvaluationFieldImpl(const TRange<FFra
 
 FTransform UMovieSceneRootMotionSection::GetAccumulatedOffsetAt(FFrameTime Time) const
 {
-	if (GapBehavior != EMovieSceneRootMotionGapBehavior::PersistPreviousTransform)
-	{
-		return FTransform::Identity;
-	}
+	FTransform Result = FTransform::Identity;
+	FFrameNumber TimeFrame = Time.FloorToFrame();
 
-	// Check if the cache needs rebuilding (track's sections have changed)
-	if (UMovieSceneCommonAnimationTrack* Track = GetTypedOuter<UMovieSceneCommonAnimationTrack>())
+	for (const FRootMotionAccumulatedOffset& Offset : AccumulatedOffsetCache)
 	{
-		if (Track->RootMotionParams.bRootMotionsDirty)
+		if (Offset.StartFrame <= TimeFrame)
 		{
-			RebuildAccumulatedOffsetCache();
-			Track->RootMotionParams.bRootMotionsDirty = false;
+			Result = Offset.AccumulatedOffset;
+		}
+		else
+		{
+			break;
 		}
 	}
 
-	// Find the most recent accumulated offset before the given time
+	return Result;
+}
+
+FTransform UMovieSceneRootMotionSection::GetAnimSpaceAccumulatedOffsetAt(FFrameTime Time) const
+{
 	FTransform Result = FTransform::Identity;
 	FFrameNumber TimeFrame = Time.FloorToFrame();
 
@@ -104,7 +128,7 @@ FTransform UMovieSceneRootMotionSection::GetAccumulatedOffsetAt(FFrameTime Time)
 	{
 		if (Offset.StartFrame <= TimeFrame)
 		{
-			Result = Offset.AccumulatedOffset;
+			Result = Offset.AnimSpaceAccumulatedOffset;
 		}
 		else
 		{
@@ -115,61 +139,169 @@ FTransform UMovieSceneRootMotionSection::GetAccumulatedOffsetAt(FFrameTime Time)
 	return Result;
 }
 
-void UMovieSceneRootMotionSection::RebuildAccumulatedOffsetCache() const
+void UMovieSceneRootMotionSection::RebuildAccumulatedOffsetCache(
+	UMovieSceneEntitySystemLinker* Linker,
+	UE::MovieScene::FInstanceHandle InstanceHandle,
+	UMovieSceneAnimationMixerTrack* MixerTrack)
 {
+	using namespace UE::MovieScene;
+
 	AccumulatedOffsetCache.Reset();
 
 	UMovieSceneTrack* OwningTrack = GetTypedOuter<UMovieSceneTrack>();
-	if (!OwningTrack)
+	if (!OwningTrack || !Linker || !MixerTrack)
 	{
 		return;
 	}
 
-	// Collect sections with root motion settings that have KeepState enabled
-	TArray<TPair<UMovieSceneSection*, UMovieSceneRootMotionSettingsDecoration*>> SectionsWithKeepState;
+	// Collect KeepState sections and transitions
+	TSet<UMovieSceneSection*> KeepStateSectionSet;
+	TArray<UMovieSceneAnimTransitionSectionBase*> Transitions;
 
 	for (UMovieSceneSection* Section : OwningTrack->GetAllSections())
 	{
-		if (!Section || !Section->HasEndFrame())
+		if (!Section)
+		{
+			continue;
+		}
+
+		if (auto* TransitionSection = Cast<UMovieSceneAnimTransitionSectionBase>(Section))
+		{
+			if (TransitionSection->IsValid())
+			{
+				Transitions.Add(TransitionSection);
+			}
+			continue;
+		}
+
+		if (!Section->HasEndFrame())
 		{
 			continue;
 		}
 
-		UMovieSceneRootMotionSettingsDecoration* Settings = Section->FindDecoration<UMovieSceneRootMotionSettingsDecoration>();
-		if (Settings && Settings->bKeepStateAfterSectionEnds)
+		auto* Settings = Section->FindDecoration<UMovieSceneRootMotionSettingsDecoration>();
+		if (Settings && Settings->GetKeepStateAfterSectionEnds())
 		{
-			SectionsWithKeepState.Add(TPair<UMovieSceneSection*, UMovieSceneRootMotionSettingsDecoration*>(Section, Settings));
+			KeepStateSectionSet.Add(Section);
 		}
 	}
 
-	// Sort by section end frame (when the accumulated offset takes effect)
-	SectionsWithKeepState.Sort([](const auto& A, const auto& B)
+	// Build a unified list of all cache entries sorted by start frame.
+	// Processing left to right ensures each entry's accumulated offset
+	// lookup finds correct values from previously computed entries.
+	struct FPendingEntry
 	{
-		return A.Key->GetExclusiveEndFrame() < B.Key->GetExclusiveEndFrame();
-	});
+		FFrameNumber EntryStartFrame;
+		UMovieSceneSection* Section;
+		FFrameTime EvalTime;
+	};
 
-	// Build the accumulated offset cache
-	FTransform RunningAccumulatedOffset = FTransform::Identity;
+	TArray<FPendingEntry> PendingEntries;
 
-	for (const auto& Pair : SectionsWithKeepState)
+	for (UMovieSceneSection* Section : KeepStateSectionSet)
 	{
-		UMovieSceneSection* Section = Pair.Key;
-		UMovieSceneRootMotionSettingsDecoration* Settings = Pair.Value;
+		FPendingEntry Entry;
+		Entry.EntryStartFrame = Section->GetExclusiveEndFrame();
+		Entry.Section = Section;
+		Entry.EvalTime = FFrameTime(Section->GetExclusiveEndFrame() - 1);
+		PendingEntries.Add(Entry);
+	}
 
-		FVector LocationOffset = Settings->MatchedLocationOffset;
-		FRotator RotationOffset = Settings->MatchedRotationOffset;
+	for (UMovieSceneAnimTransitionSectionBase* Transition : Transitions)
+	{
+		if (!KeepStateSectionSet.Contains(Transition->FromSection.Get()))
+		{
+			continue;
+		}
 
-		FTransform SectionOffset(RotationOffset.Quaternion(), LocationOffset);
+		TRange<FFrameNumber> TransitionRange = Transition->GetRange();
+		if (TransitionRange.IsEmpty() || !TransitionRange.HasLowerBound() || !TransitionRange.HasUpperBound())
+		{
+			continue;
+		}
 
-		// Accumulate the offset
-		RunningAccumulatedOffset = SectionOffset * RunningAccumulatedOffset;
+		FFrameNumber TransitionStart = TransitionRange.GetLowerBoundValue();
+		FFrameNumber TransitionEnd = TransitionRange.GetUpperBoundValue();
 
-		// Add to cache - this offset takes effect when the section ends
-		FRootMotionAccumulatedOffset CacheEntry;
-		CacheEntry.StartFrame = Section->GetExclusiveEndFrame();
-		CacheEntry.AccumulatedOffset = RunningAccumulatedOffset;
-		CacheEntry.Weight = 1.0f;
+		{
+			FPendingEntry Entry;
+			Entry.EntryStartFrame = TransitionStart;
+			Entry.Section = Transition->FromSection.Get();
+			Entry.EvalTime = FFrameTime(TransitionStart);
+			PendingEntries.Add(Entry);
+		}
 
-		AccumulatedOffsetCache.Add(CacheEntry);
+		if (Transition->FromSection
+			&& Transition->FromSection->HasEndFrame()
+			&& TransitionEnd != Transition->FromSection->GetExclusiveEndFrame())
+		{
+			FPendingEntry Entry;
+			Entry.EntryStartFrame = TransitionEnd;
+			Entry.Section = Transition->FromSection.Get();
+			Entry.EvalTime = FFrameTime(TransitionEnd - 1);
+			PendingEntries.Add(Entry);
+		}
+	}
+
+	PendingEntries.Sort([](const FPendingEntry& A, const FPendingEntry& B)
+	{
+		return A.EntryStartFrame < B.EntryStartFrame;
+	});
+
+	for (const FPendingEntry& Pending : PendingEntries)
+	{
+		auto* Settings = Pending.Section->FindDecoration<UMovieSceneRootMotionSettingsDecoration>();
+		const bool bUsesAccumulatedOffset = Settings
+			&& Settings->GetRootTransformMode() == EMovieSceneRootMotionTransformMode::AccumulatedOffset;
+
+		if (bUsesAccumulatedOffset && Settings->bInitialRootTransformValid)
+		{
+			// Bake in animation space (skip conversion), then combine
+			// with the accumulated offset from earlier entries.
+			AnimMixerBakeEvaluation::FBakeFilter Filter;
+			Filter.IncludeOnlySections.Add(FObjectKey(Pending.Section));
+			Filter.bSkipRootMotionConversion = true;
+
+			AnimMixerBakeEvaluation::FBakeResult Result = AnimMixerBakeEvaluation::EvaluateAtTime(
+				Linker, InstanceHandle, MixerTrack, Pending.EvalTime, Filter);
+
+			FTransform Delta = Settings->InitialRootTransform.Inverse() * Result.RootMotionTransform;
+
+			FRootMotionAccumulatedOffset CacheEntry;
+			CacheEntry.StartFrame = Pending.EntryStartFrame;
+			CacheEntry.Weight = 1.0f;
+
+			if (Pending.Section->HasStartFrame())
+			{
+				FFrameTime StartTime(Pending.Section->GetInclusiveStartFrame());
+				CacheEntry.AccumulatedOffset = Delta * GetAccumulatedOffsetAt(StartTime);
+				CacheEntry.AnimSpaceAccumulatedOffset = Delta * GetAnimSpaceAccumulatedOffsetAt(StartTime);
+			}
+			else
+			{
+				CacheEntry.AccumulatedOffset = Delta;
+				CacheEntry.AnimSpaceAccumulatedOffset = Delta;
+			}
+			AccumulatedOffsetCache.Add(CacheEntry);
+		}
+		else
+		{
+			// Non-AccumulatedOffset section: full bake with world conversion.
+			// AnimationSpaceRootMotion is captured by the conversion task.
+			AnimMixerBakeEvaluation::FBakeFilter Filter;
+			Filter.IncludeOnlySections.Add(FObjectKey(Pending.Section));
+
+			AnimMixerBakeEvaluation::FBakeResult Result = AnimMixerBakeEvaluation::EvaluateAtTime(
+				Linker, InstanceHandle, MixerTrack, Pending.EvalTime, Filter);
+
+			FRootMotionAccumulatedOffset CacheEntry;
+			CacheEntry.StartFrame = Pending.EntryStartFrame;
+			CacheEntry.AccumulatedOffset = Result.RootMotionTransform;
+			CacheEntry.AnimSpaceAccumulatedOffset = Result.AnimationSpaceRootMotion;
+			CacheEntry.Weight = 1.0f;
+			AccumulatedOffsetCache.Add(CacheEntry);
+		}
 	}
+
+	bAccumulatedOffsetCacheDirty = false;
 }
diff --git a/Engine/Plugins/Experimental/MovieSceneAnimMixer/Source/MovieSceneAnimMixer/Private/Systems/MovieSceneAnimMixerSystem.cpp b/Engine/Plugins/Experimental/MovieSceneAnimMixer/Source/MovieSceneAnimMixer/Private/Systems/MovieSceneAnimMixerSystem.cpp
--- a/Engine/Plugins/Experimental/MovieSceneAnimMixer/Source/MovieSceneAnimMixer/Private/Systems/MovieSceneAnimMixerSystem.cpp
+++ b/Engine/Plugins/Experimental/MovieSceneAnimMixer/Source/MovieSceneAnimMixer/Private/Systems/MovieSceneAnimMixerSystem.cpp
@@ -2,8 +2,11 @@
 
 #include "Systems/MovieSceneAnimMixerSystem.h"
 
+#include "MovieSceneAnimationMixerTrack.h"
 #include "MovieSceneRootMotionSection.h"
+#include "MovieSceneRootMotionTargetDecoration.h"
 #include "Systems/ByteChannelEvaluatorSystem.h"
+#include "Systems/MovieSceneMixerBakeTargetSystem.h"
 
 #include "EntitySystem/BuiltInComponentTypes.h"
 #include "AnimMixerComponentTypes.h"
@@ -49,387 +52,721 @@
 #include "MovieSceneAnimTransitionSectionBase.h"
 #include "EvaluationVM/Tasks/TransitionEvaluationTask.h"
 #include "EntitySystem/MovieSceneInstanceRegistry.h"
-
 #include UE_INLINE_GENERATED_CPP_BY_NAME(MovieSceneAnimMixerSystem)
 
-
 namespace UE::MovieScene
 {
-	struct FUpdateTaskEntities
+
+// Filter for controlling which mixer entries are included in program building.
+struct FProgramBuildFilter
+{
+	// Entries whose EntityOwner matches are skipped
+	const TSet<FObjectKey>* ExcludeSections = nullptr;
+
+	// If non-null, only entries whose EntityOwner matches are included
+	const TSet<FObjectKey>* IncludeOnly = nullptr;
+
+	// If >= 0, only entries with priority <= MaxPriority are included
+	int32 MaxPriority = -1;
+
+	// When true, skip ConvertRootMotionToWorldSpaceTask — result is in animation space
+	bool bSkipRootMotionConversion = false;
+
+	// When set, the conversion task writes the pre-conversion animation-space root here
+	TSharedPtr<FTransform> CaptureAnimSpaceRoot;
+};
+
+// Context needed for building evaluation programs.
+struct FProgramBuildContext
+{
+	TSharedPtr<FMovieSceneMixerRootMotionComponentData> RootMotion;
+	const UMovieSceneRootMotionSystem* RootMotionSystem = nullptr;
+	const UMovieSceneTransformOriginSystem* TransformOriginSystem = nullptr;
+	UMovieSceneEntitySystemLinker* Linker = nullptr;
+	USceneComponent* Component = nullptr;
+};
+
+// Returns true if the entry passes the filter (i.e. should be considered for inclusion)
+static bool PassesFilter(const FMovieSceneAnimMixerEntry& Entry, const FProgramBuildFilter& Filter)
+{
+	if (Filter.ExcludeSections && Filter.ExcludeSections->Contains(Entry.EntityOwner))
 	{
-		static void ForEachEntity(
-			TSharedPtr<FAnimNextEvaluationTask> Task,
-			TSharedPtr<FMovieSceneAnimMixerEntry> MixerEntry,
-			const FMovieSceneRootMotionSettings* RootMotionSettings,
-			const double* WeightAndEasings)
+		return false;
+	}
+	if (Filter.IncludeOnly && !Filter.IncludeOnly->Contains(Entry.EntityOwner))
+	{
+		return false;
+	}
+	if (Filter.MaxPriority >= 0 && Entry.Priority > Filter.MaxPriority)
+	{
+		return false;
+	}
+	return true;
+}
+
+// Builds the ConvertRootMotionToWorldSpace task for a single mixer entry
+// based on its root motion settings. Returns an empty optional if no
+// conversion is needed.
+static TOptional<FAnimNextConvertRootMotionToWorldSpaceTask> BuildSpaceConversionTask(
+	const FMovieSceneAnimMixerEntry& MixerEntry,
+	const FProgramBuildContext& Context,
+	bool bRootComponentHasKeyedTransform)
+{
+	using ESpaceConversions = FAnimNextConvertRootMotionToWorldSpaceTask::ESpaceConversions;
+
+	FTransform TransformOrigin  = FTransform::Identity;
+	FTransform RootTransform    = FTransform::Identity;
+	FVector    RootOffsetOrigin = FVector(ForceInitToZero);
+
+	ESpaceConversions Conversion = ESpaceConversions::None;
+	EMovieSceneRootMotionSpace ThisRootSpace = EMovieSceneRootMotionSpace::AnimationSpace;
+
+	bool bUseAccumulatedOffset = false;
+	FTransform ResolvedAccumulatedTransform = FTransform::Identity;
+	FTransform ResolvedInitialRootTransform = FTransform::Identity;
+
+	if (MixerEntry.RootMotionSettings)
+	{
+		ThisRootSpace = MixerEntry.RootMotionSettings->RootMotionSpace;
+
+		const bool bHasLocation = !MixerEntry.RootMotionSettings->RootLocation.IsZero();
+		const bool bHasRotation = !MixerEntry.RootMotionSettings->RootRotation.IsIdentity();
+
+		if (MixerEntry.RootMotionSettings->TransformMode == EMovieSceneRootMotionTransformMode::Override)
 		{
-			if (MixerEntry)
+			Conversion |= ESpaceConversions::RootTransformOverride;
+			RootTransform = FTransform(
+				MixerEntry.RootMotionSettings->RootRotation,
+				MixerEntry.RootMotionSettings->RootLocation);
+		}
+		else if (MixerEntry.RootMotionSettings->TransformMode == EMovieSceneRootMotionTransformMode::Offset)
+		{
+			if (bHasLocation || bHasRotation)
 			{
-				MixerEntry->EvalTask = Task;
+				Conversion |= ESpaceConversions::RootTransformOffset;
+				RootOffsetOrigin = MixerEntry.RootMotionSettings->RootOriginLocation;
+				RootTransform = FTransform(
+					MixerEntry.RootMotionSettings->RootRotation,
+					MixerEntry.RootMotionSettings->RootLocation);
+			}
+		}
+		else if (MixerEntry.RootMotionSettings->TransformMode == EMovieSceneRootMotionTransformMode::AccumulatedOffset)
+		{
+			UMovieSceneSection* OwnerSection = Cast<UMovieSceneSection>(MixerEntry.EntityOwner.ResolveObjectPtr());
+			UMovieSceneAnimationMixerTrack* MixerTrack = OwnerSection
+				? OwnerSection->GetTypedOuter<UMovieSceneAnimationMixerTrack>() : nullptr;
+			if (MixerTrack)
+			{
+				if (auto* TargetDecoration = MixerTrack->FindDecoration<UMovieSceneRootMotionTargetDecoration>())
+				{
+					if (auto* RootMotionSection = TargetDecoration->GetRootMotionSection())
+					{
+						FFrameTime SectionStart(OwnerSection->HasStartFrame()
+							? OwnerSection->GetInclusiveStartFrame() : FFrameNumber(0));
+						FFrameNumber StartFrame = SectionStart.FloorToFrame();
 
-				// If bSkipPoseWeight is set, use 1.0 for weight
-				MixerEntry->PoseWeight = (MixerEntry->bSkipPoseWeight || !WeightAndEasings) ? 1.0 : *WeightAndEasings;
+						for (const FRootMotionAccumulatedOffset& Offset : RootMotionSection->AccumulatedOffsetCache)
+						{
+							if (Offset.StartFrame <= StartFrame)
+							{
+								bUseAccumulatedOffset = true;
+							}
+						}
 
-				if (RootMotionSettings)
+						if (bUseAccumulatedOffset)
+						{
+							Conversion |= ESpaceConversions::AccumulatedOffset;
+							ResolvedAccumulatedTransform =
+								RootMotionSection->GetAccumulatedOffsetAt(SectionStart);
+						}
+					}
+				}
+
+				if (bUseAccumulatedOffset)
 				{
-					MixerEntry->RootMotionSettings = *RootMotionSettings;
+					if (auto* Settings = OwnerSection->FindDecoration<UMovieSceneRootMotionSettingsDecoration>())
+					{
+						if (Settings->bInitialRootTransformValid)
+						{
+							ResolvedInitialRootTransform = Settings->InitialRootTransform;
+						}
+					}
 				}
 			}
+
+			if (bHasLocation || bHasRotation)
+			{
+				Conversion |= ESpaceConversions::RootTransformOffset;
+				RootOffsetOrigin = MixerEntry.RootMotionSettings->RootOriginLocation;
+				RootTransform = FTransform(
+					MixerEntry.RootMotionSettings->RootRotation,
+					MixerEntry.RootMotionSettings->RootLocation);
+			}
 		}
-	};
+		else if (MixerEntry.RootMotionSettings->TransformMode == EMovieSceneRootMotionTransformMode::BoneMatch)
+		{
+			if (MixerEntry.BoneMatchTransform.IsSet())
+			{
+				Conversion |= ESpaceConversions::BoneMatchOffset;
+			}
 
-	struct FEvaluateAnimMixers
+			if (bHasLocation || bHasRotation)
+			{
+				Conversion |= ESpaceConversions::RootTransformOffset;
+				RootOffsetOrigin = MixerEntry.RootMotionSettings->RootOriginLocation;
+				RootTransform = FTransform(
+					MixerEntry.RootMotionSettings->RootRotation,
+					MixerEntry.RootMotionSettings->RootLocation);
+			}
+		}
+	}
+
+	if (ThisRootSpace == EMovieSceneRootMotionSpace::AnimationSpace)
 	{
-		TMap<FMovieSceneAnimMixerKey, TSharedPtr<FMovieSceneAnimMixer>>* Mixers;
-		TMap<FObjectKey, TWeakPtr<FMovieSceneMixerRootMotionComponentData>>* RootMotions;
-		const UMovieSceneRootMotionSystem* RootMotionSystem;
-		const UMovieSceneTransformOriginSystem* TransformOriginSystem;
-		UMovieSceneEntitySystemLinker* Linker;
+		Conversion |= ESpaceConversions::ComponentToActorRotation;
 
-		FEvaluateAnimMixers(TMap<FMovieSceneAnimMixerKey, TSharedPtr<FMovieSceneAnimMixer>>* InMixers, TMap<FObjectKey, TWeakPtr<FMovieSceneMixerRootMotionComponentData>>* InRootMotions, UMovieSceneEntitySystemLinker* InLinker)
-			: Mixers(InMixers)
-			, RootMotions(InRootMotions)
-			, Linker(InLinker)
+		if (!bUseAccumulatedOffset)
 		{
-			RootMotionSystem = InLinker->FindSystem<UMovieSceneRootMotionSystem>();
-			TransformOriginSystem = InLinker->FindSystem<UMovieSceneTransformOriginSystem>();
+			if (bRootComponentHasKeyedTransform)
+			{
+				Conversion |= ESpaceConversions::AnimationToWorld;
+			}
+			else if (Context.TransformOriginSystem && Context.TransformOriginSystem->GetTransformOrigin(MixerEntry.InstanceHandle, TransformOrigin))
+			{
+				Conversion |= ESpaceConversions::TransformOriginToWorld;
+			}
 		}
+	}
+	else
+	{
+		Conversion |= ESpaceConversions::WorldSpaceComponentTransformCompensation;
+	}
 
-		void ForEachEntity(
-			FObjectComponent MeshComponent,
-			const TInstancedStruct<FMovieSceneMixedAnimationTarget>& Target,
-			TSharedPtr<FAnimNextEvaluationTask>& MixerTask
-			) const
+	if (Conversion != ESpaceConversions::None)
+	{
+		FAnimNextConvertRootMotionToWorldSpaceTask Task =
+			FAnimNextConvertRootMotionToWorldSpaceTask::Make(
+				Context.RootMotion, TransformOrigin, RootTransform, RootOffsetOrigin, Conversion);
+
+		if (bUseAccumulatedOffset)
 		{
-			using ESpaceConversions = FAnimNextConvertRootMotionToWorldSpaceTask::ESpaceConversions;
+			Task.AccumulatedTransform = ResolvedAccumulatedTransform;
+			Task.InitialAnimRootTransform = ResolvedInitialRootTransform;
+		}
 
-			TSharedPtr<FMovieSceneAnimMixer> Mixer = Mixers->FindRef(FMovieSceneAnimMixerKey({ MeshComponent.GetObject(), Target }));
-			if (!Mixer)
+		if (EnumHasAnyFlags(Conversion, ESpaceConversions::BoneMatchOffset) && MixerEntry.BoneMatchTransform.IsSet())
+		{
+			Task.BoneMatchTransform = MixerEntry.BoneMatchTransform.GetValue();
+		}
+
+		return Task;
+	}
+
+	return {};
+}
+
+// Build an evaluation program from a mixer's entries with optional filtering.
+// Handles transition detection (transitions whose from/to sections are filtered out
+// are ignored, and the surviving section is treated normally), priority grouping,
+// blend stacks, root motion space conversion, and mirroring.
+// Does NOT include FAnimNextStoreRootTransformTask - caller adds that if needed.
+static TSharedPtr<UE::UAF::FEvaluationProgram> BuildProgramFromEntries(
+	const FMovieSceneAnimMixer& Mixer,
+	const FProgramBuildContext& Context,
+	const FProgramBuildFilter& Filter = {})
+{
+	// First pass: determine which non-transition entries pass the filter.
+	// We need this to decide which transitions are valid.
+	TSet<FObjectKey> IncludedSectionOwners;
+	for (const TWeakPtr<FMovieSceneAnimMixerEntry>& WeakEntry : Mixer.WeakEntries)
+	{
+		TSharedPtr<FMovieSceneAnimMixerEntry> Entry = WeakEntry.Pin();
+		if (!Entry || !Entry->EvalTask.IsValid() || Entry->bIsTransition)
+		{
+			continue;
+		}
+		if (PassesFilter(*Entry, Filter))
+		{
+			IncludedSectionOwners.Add(Entry->EntityOwner);
+		}
+	}
+
+	USceneComponent* Component = Context.Component;
+	USceneComponent* Root = Component ? Component->GetOwner()->GetRootComponent() : nullptr;
+	const bool bRootComponentHasKeyedTransform = (Root && Context.RootMotionSystem && Context.RootMotionSystem->IsTransformKeyed(Root));
+
+	// Second pass: handle transitions. A transition is only valid if BOTH its from and to
+	// sections are in the included set. If a transition is invalid (one side is filtered out),
+	// the surviving section is treated normally (not added to SectionsInTransition).
+	TSet<FObjectKey> SectionsInTransition;
+
+	for (const TWeakPtr<FMovieSceneAnimMixerEntry>& WeakEntry : Mixer.WeakEntries)
+	{
+		TSharedPtr<FMovieSceneAnimMixerEntry> Entry = WeakEntry.Pin();
+		if (!Entry || !Entry->bIsTransition || !Entry->EvalTask.IsValid())
+		{
+			continue;
+		}
+		if (!PassesFilter(*Entry, Filter))
+		{
+			continue;
+		}
+
+		UMovieSceneAnimTransitionSectionBase* TransitionSection = Cast<UMovieSceneAnimTransitionSectionBase>(Entry->EntityOwner.ResolveObjectPtr());
+		if (!TransitionSection || !TransitionSection->IsValid())
+		{
+			continue;
+		}
+
+		FObjectKey FromKey(TransitionSection->FromSection);
+		FObjectKey ToKey(TransitionSection->ToSection);
+
+		// Only activate this transition if both from/to are included
+		if (!IncludedSectionOwners.Contains(FromKey) || !IncludedSectionOwners.Contains(ToKey))
+		{
+			continue;
+		}
+
+		SectionsInTransition.Add(FromKey);
+		SectionsInTransition.Add(ToKey);
+
+		TSharedPtr<FMovieSceneAnimMixerEntry> FromEntry = Mixer.FindEntryByOwner(FromKey);
+		TSharedPtr<FMovieSceneAnimMixerEntry> ToEntry = Mixer.FindEntryByOwner(ToKey);
+
+		// Wrap from/to tasks in mini-programs that include mirroring and
+		// space conversion so the transition blends in world space. This is
+		// required for AccumulatedOffset (different starting positions) and
+		// is more correct in general when sections have different settings.
+		auto WrapEntryTask = [&](const TSharedPtr<FMovieSceneAnimMixerEntry>& WrapEntry)
+			-> TSharedPtr<FAnimNextEvaluationTask>
+		{
+			if (!WrapEntry || !WrapEntry->EvalTask.IsValid())
 			{
-				return;
+				return nullptr;
 			}
 
-			if (Mixer->WeakEntries.Num() == 0)
+			bool bNeedsMirror = (WrapEntry->MirrorTable != nullptr);
+			TOptional<FAnimNextConvertRootMotionToWorldSpaceTask> SpaceTask;
+			if (!Filter.bSkipRootMotionConversion)
 			{
-				MixerTask = nullptr;
-				Mixer->EvaluationProgram = nullptr;
-				return;
+				SpaceTask = BuildSpaceConversionTask(*WrapEntry, Context, bRootComponentHasKeyedTransform);
 			}
 
-			USceneComponent* Component = Cast<USceneComponent>(MeshComponent.GetObject());
-			USceneComponent* Root = Component ? Component->GetOwner()->GetRootComponent() : nullptr;
+			if (!bNeedsMirror && !SpaceTask.IsSet())
+			{
+				return WrapEntry->EvalTask;
+			}
 
-			const bool bComponentHasKeyedTransform = (Component && RootMotionSystem->IsTransformKeyed(Component));
-			const bool bRootComponentHasKeyedTransform = (Root && RootMotionSystem->IsTransformKeyed(Root));
+			TSharedPtr<UE::UAF::FEvaluationProgram> MiniProgram = MakeShared<UE::UAF::FEvaluationProgram>();
+			MiniProgram->AppendTaskPtr(WrapEntry->EvalTask);
 
-			// Create a new eval program
-			Mixer->EvaluationProgram = MakeShared<UE::UAF::FEvaluationProgram>();
+			if (bNeedsMirror)
+			{
+				TSharedPtr<FAnimNextEvaluationMirroringTask> MirrorTask = MakeShared<FAnimNextEvaluationMirroringTask>();
+				MirrorTask->Setup.MirrorDataTable = WrapEntry->MirrorTable;
+				MiniProgram->AppendTaskPtr(MirrorTask);
+			}
 
-			if (!MixerTask.IsValid())
+			if (Filter.CaptureAnimSpaceRoot && SpaceTask.IsSet())
 			{
-				MixerTask = MakeShared<FAnimNextExecuteProgramTask>();
+				SpaceTask->CaptureAnimSpaceRoot = Filter.CaptureAnimSpaceRoot;
 			}
-			TSharedPtr<FAnimNextExecuteProgramTask> EvalProgramTask = StaticCastSharedPtr<FAnimNextExecuteProgramTask>(MixerTask);
-			EvalProgramTask->Program = Mixer->EvaluationProgram;
-
-			// From lowest to highest priority, add the tasks to the program, grouped by their priority.
-			// Tasks within the same priority are blending using a weighted average, which is then collapsed into a single pose, and blended with the
-			//      subsequnt priority based on its total accumulated PoseWeight.
-			//
-			// For example, 5 animations with a structure of:
-			// Priority 0: (A, w:0.5), (B, w:0.5), (C, w:0.5)
-			// Priority 2: (D, w:0.25), (E, w:0.5)
-			//
-			// Would result in a blend stack equivalent to the following operation:
-			// P1 = A*(0.5/1.5) + B*(0.5/1.5) + C(0.5/1.5)
-			// P2 = D*(0.25/0.75) + B*(0.5/0.75)
-			//
-			// FinalPose = (0.25*P1) + (0.75*P2)
-			//
-			// Note that we assume here that the target will push some kind of 'base pose' to the VM before evaluating the task. This allows us to blend in/out from gameplay for example.
-
-			const int32 NumMixerEntries = Mixer->WeakEntries.Num();
-
-			// ----------------------------------------------------------------------------------------------------------------------------------------
-			// Transition handling: Build transition tasks and track which sections are handled by transitions
-			// Sections handled by transitions will be skipped in the main evaluation loop below.
-			TSet<FObjectKey> SectionsInTransition;
-
-			for (int32 Index = 0; Index < NumMixerEntries; ++Index)
+
+			if (SpaceTask.IsSet())
 			{
-				TSharedPtr<FMovieSceneAnimMixerEntry> Entry = Mixer->WeakEntries[Index].Pin();
-				if (!Entry || !Entry->bIsTransition)
+				MiniProgram->AppendTask(MoveTemp(SpaceTask.GetValue()));
+			}
+
+			TSharedPtr<FAnimNextExecuteProgramTask> WrappedTask = MakeShared<FAnimNextExecuteProgramTask>();
+			WrappedTask->Program = MiniProgram;
+			return WrappedTask;
+		};
+
+		TSharedPtr<FAnimNextEvaluationTask> FromTask = WrapEntryTask(FromEntry);
+		TSharedPtr<FAnimNextEvaluationTask> ToTask = WrapEntryTask(ToEntry);
+
+		if (FAnimNextTransitionEvaluationTask* TransitionTask =
+			static_cast<FAnimNextTransitionEvaluationTask*>(Entry->EvalTask.Get()))
+		{
+			float DeltaTime = 0.0f;
+			if (Context.Linker && Entry->InstanceHandle.IsValid())
+			{
+				const FMovieSceneContext& MovieContext = Context.Linker->GetInstanceRegistry()->GetContext(Entry->InstanceHandle);
+				const FFrameTime FrameDelta = MovieContext.GetDelta();
+				const FFrameRate FrameRate = MovieContext.GetFrameRate();
+				DeltaTime = static_cast<float>(FrameDelta.AsDecimal() / FrameRate.AsDecimal());
+			}
+
+			double BlendWeight = Entry->PoseWeight;
+			TransitionTask->Update(FromTask, ToTask, BlendWeight, DeltaTime);
+		}
+
+		Entry->PoseWeight = 1.0;
+	}
+
+	// Third pass: collect all entries that will actually be evaluated.
+	// Non-transition entries that pass filter AND aren't handled by a transition,
+	// plus transition entries whose from/to are both valid.
+	TArray<TSharedPtr<FMovieSceneAnimMixerEntry>> FinalEntries;
+	for (const TWeakPtr<FMovieSceneAnimMixerEntry>& WeakEntry : Mixer.WeakEntries)
+	{
+		TSharedPtr<FMovieSceneAnimMixerEntry> Entry = WeakEntry.Pin();
+		if (!Entry || !Entry->EvalTask.IsValid())
+		{
+			continue;
+		}
+		if (!PassesFilter(*Entry, Filter))
+		{
+			continue;
+		}
+		// Skip non-transition entries that are handled by a valid transition
+		if (!Entry->bIsTransition && SectionsInTransition.Contains(Entry->EntityOwner))
+		{
+			continue;
+		}
+		// Skip transition entries whose from/to aren't both included
+		if (Entry->bIsTransition)
+		{
+			UMovieSceneAnimTransitionSectionBase* TransitionSection = Cast<UMovieSceneAnimTransitionSectionBase>(Entry->EntityOwner.ResolveObjectPtr());
+			if (!TransitionSection || !TransitionSection->IsValid())
+			{
+				continue;
+			}
+			if (!IncludedSectionOwners.Contains(FObjectKey(TransitionSection->FromSection)) ||
+				!IncludedSectionOwners.Contains(FObjectKey(TransitionSection->ToSection)))
+			{
+				continue;
+			}
+		}
+		FinalEntries.Add(Entry);
+	}
+
+	if (FinalEntries.Num() == 0)
+	{
+		return nullptr;
+	}
+
+	TSharedPtr<UE::UAF::FEvaluationProgram> Program = MakeShared<UE::UAF::FEvaluationProgram>();
+
+	// Lookahead: compute total absolute weight for a priority group
+	const int32 NumEntries = FinalEntries.Num();
+
+	auto LookaheadComputeAccumulatedWeight = [&FinalEntries, NumEntries](int32 StartAtIndex, int32 Priority, float& OutAccumulatedWeight) -> bool
+	{
+		int32 NumAbsoluteBlends = 0;
+		int32 NumSkippedBlends = 0;
+		float AccumulatedWeight = 0.f;
+		for (int32 NextIndex = StartAtIndex; NextIndex < NumEntries; ++NextIndex)
+		{
+			const TSharedPtr<FMovieSceneAnimMixerEntry>& NextEntry = FinalEntries[NextIndex];
+			if (!NextEntry->EvalTask.IsValid())
+			{
+				continue;
+			}
+			if (NextEntry->Priority != Priority || NextEntry->bAdditive)
+			{
+				break;
+			}
+			else if (NextEntry->bRequiresBlend && !NextEntry->bAdditive)
+			{
+				if (!FMath::IsNearlyEqual(AccumulatedWeight, AccumulatedWeight + NextEntry->PoseWeight, KINDA_SMALL_NUMBER))
 				{
-					continue;
+					++NumAbsoluteBlends;
+					AccumulatedWeight += NextEntry->PoseWeight;
 				}
-
-				// Get the transition section from the EntityOwner
-				UMovieSceneAnimTransitionSectionBase* TransitionSection = Cast<UMovieSceneAnimTransitionSectionBase>(Entry->EntityOwner.ResolveObjectPtr());
-				if (!TransitionSection || !TransitionSection->IsValid())
+				else
 				{
-					continue;
+					++NumSkippedBlends;
 				}
+			}
+		}
 
-				// Mark from/to sections as handled by this transition
-				SectionsInTransition.Add(FObjectKey(TransitionSection->FromSection));
-				SectionsInTransition.Add(FObjectKey(TransitionSection->ToSection));
+		if (NumAbsoluteBlends + NumSkippedBlends > 1 && !FMath::IsNearlyZero(AccumulatedWeight, KINDA_SMALL_NUMBER))
+		{
+			OutAccumulatedWeight = AccumulatedWeight;
+			return true;
+		}
+		return false;
+	};
 
-				// Find from/to entries to get their tasks
-				TSharedPtr<FMovieSceneAnimMixerEntry> FromEntry = Mixer->FindEntryByOwner(FObjectKey(TransitionSection->FromSection));
-				TSharedPtr<FMovieSceneAnimMixerEntry> ToEntry = Mixer->FindEntryByOwner(FObjectKey(TransitionSection->ToSection));
+	// Main evaluation loop - iterate entries grouped by priority
+	for (int32 Index = 0; Index < NumEntries; )
+	{
+		const TSharedPtr<FMovieSceneAnimMixerEntry>& PeekEntry = FinalEntries[Index];
+		const int32 Priority = PeekEntry->Priority;
 
-				TSharedPtr<FAnimNextEvaluationTask> FromTask = FromEntry ? FromEntry->EvalTask : nullptr;
-				TSharedPtr<FAnimNextEvaluationTask> ToTask = ToEntry ? ToEntry->EvalTask : nullptr;
+		float AccumulatedPriorityWeight = 1.f;
+		const bool bNeedsSeparateWeightStack = LookaheadComputeAccumulatedWeight(Index, Priority, AccumulatedPriorityWeight);
 
-				// Update the transition task with from/to tasks and blend weight.
-				// The task was created during entity import and persists across frames.
-				if (FAnimNextTransitionEvaluationTask* TransitionTask =
-					static_cast<FAnimNextTransitionEvaluationTask*>(Entry->EvalTask.Get()))
-				{
-					// Get delta time from the instance context
-					float DeltaTime = 0.0f;
-					if (Entry->InstanceHandle.IsValid())
-					{
-						const FMovieSceneContext& Context = Linker->GetInstanceRegistry()->GetContext(Entry->InstanceHandle);
-						const FFrameTime FrameDelta = Context.GetDelta();
-						const FFrameRate FrameRate = Context.GetFrameRate();
-						DeltaTime = static_cast<float>(FrameDelta.AsDecimal() / FrameRate.AsDecimal());
-					}
+		bool bIsFirstAbsoluteBlend = true;
+		for ( ; Index < NumEntries; ++Index)
+		{
+			const TSharedPtr<FMovieSceneAnimMixerEntry>& MixerEntry = FinalEntries[Index];
+			if (MixerEntry->Priority != Priority)
+			{
+				break;
+			}
 
-					// The blend weight comes from the transition's BlendCurve via the WeightChannel evaluation
-					double BlendWeight = Entry->PoseWeight;
-					TransitionTask->Update(FromTask, ToTask, BlendWeight, DeltaTime);
-				}
+			if (!MixerEntry->EvalTask.IsValid())
+			{
+				continue;
+			}
+
+			if (bNeedsSeparateWeightStack && FMath::IsNearlyZero(MixerEntry->PoseWeight, KINDA_SMALL_NUMBER))
+			{
+				continue;
+			}
 
-				// Reset PoseWeight to 1.0 since the transition replaces the from/to sections
-				// and should blend into the overall mix at full weight
-				Entry->PoseWeight = 1.0;
+			// Evaluate the pose
+			Program->AppendTaskPtr(MixerEntry->EvalTask);
+
+			// Mirror if specified
+			if (MixerEntry->MirrorTable)
+			{
+				TSharedPtr<FAnimNextEvaluationMirroringTask> MirrorTask = MakeShared<FAnimNextEvaluationMirroringTask>();
+				MirrorTask->Setup.MirrorDataTable = MixerEntry->MirrorTable;
+				Program->AppendTaskPtr(MirrorTask);
 			}
-			// ----------------------------------------------------------------------------------------------------------------------------------------
 
-			// Function that performs a lookahead accumulation of the total number of absolute weights for the specified start index's priority,
-			//  and returns a value indicating whether the entries need a separate blend stack (true) or whether they can be blended as a single operation
-			auto LookaheadComputeAccumulatedWeight = [&Mixer, NumMixerEntries, &SectionsInTransition](int32 StartAtIndex, int32 Priority, float& OutAccumulatedWeight) -> bool
+			// Handle root space conversions. Skipped for bake (animation-space
+			// results) and for transitions (from/to are already converted to
+			// world space inside the wrapped mini-programs).
+			if (!Filter.bSkipRootMotionConversion && !MixerEntry->bIsTransition)
 			{
-				int32 NumAbsoluteBlends = 0;
-				int32 NumSkippedBlends = 0;
-				float AccumulatedWeight = 0.f;
-				for (int32 NextIndex = StartAtIndex; NextIndex < NumMixerEntries; ++NextIndex)
+				TOptional<FAnimNextConvertRootMotionToWorldSpaceTask> SpaceTask =
+					BuildSpaceConversionTask(*MixerEntry, Context, bRootComponentHasKeyedTransform);
+				if (SpaceTask.IsSet())
 				{
-					TSharedPtr<FMovieSceneAnimMixerEntry> NextMixerEntry = Mixer->WeakEntries[NextIndex].Pin();
-					if (!ensure(NextMixerEntry))
-					{
-						continue;
-					}
-					if (!NextMixerEntry->EvalTask.IsValid())
-					{
-						continue;
-					}
-					// Skip entries whose source sections are handled by an active transition
-					if (SectionsInTransition.Contains(NextMixerEntry->EntityOwner))
+					if (Filter.CaptureAnimSpaceRoot)
 					{
-						continue;
+						SpaceTask->CaptureAnimSpaceRoot = Filter.CaptureAnimSpaceRoot;
 					}
-					if (NextMixerEntry->Priority != Priority || NextMixerEntry->bAdditive)
+					Program->AppendTask(MoveTemp(SpaceTask.GetValue()));
+				}
+			}
+
+			if (MixerEntry->bRequiresBlend)
+			{
+				if (MixerEntry->bAdditive)
+				{
+					Program->AppendTask(FAnimNextApplyAdditiveKeyframeTask::Make(MixerEntry->PoseWeight));
+				}
+				else if (bNeedsSeparateWeightStack)
+				{
+					if (bIsFirstAbsoluteBlend)
 					{
-						break;
+						bIsFirstAbsoluteBlend = false;
+						Program->AppendTask(FAnimNextBlendOverwriteKeyframeWithScaleTask::Make(MixerEntry->PoseWeight / AccumulatedPriorityWeight));
 					}
-					else if (NextMixerEntry->bRequiresBlend && !NextMixerEntry->bAdditive)
+					else
 					{
-						if (!FMath::IsNearlyEqual(AccumulatedWeight, AccumulatedWeight + NextMixerEntry->PoseWeight, KINDA_SMALL_NUMBER))
-						{
-							++NumAbsoluteBlends;
-							AccumulatedWeight += NextMixerEntry->PoseWeight;
-						}
-						else
-						{
-							++NumSkippedBlends;
-						}
+						Program->AppendTask(FMovieSceneAccumulateAbsoluteBlendTask::Make(MixerEntry->PoseWeight / AccumulatedPriorityWeight));
 					}
 				}
-
-				// When skipping blends, the accumulated weight still needs to be factored in for the absolute blends that aren't skipped,
-				// so if the total is greater than one, use a separate blend stack, and the accumulated weight.
-				if (NumAbsoluteBlends + NumSkippedBlends > 1 && !FMath::IsNearlyZero(AccumulatedWeight, KINDA_SMALL_NUMBER))
+				else
 				{
-					OutAccumulatedWeight = AccumulatedWeight;
-					return true;
+					Program->AppendTask(FAnimNextBlendTwoKeyframesPreserveRootMotionTask::Make(MixerEntry->PoseWeight));
 				}
-				return false;
-			};
+			}
+		}
 
-			TSharedPtr<FMovieSceneMixerRootMotionComponentData> RootMotion = RootMotions->FindRef(MeshComponent.GetObject()).Pin();
+		if (bNeedsSeparateWeightStack)
+		{
+			Program->AppendTask(FAnimNextNormalizeKeyframeRotationsTask());
+			constexpr float WeightedAverageBlendWeight = 1.f;
+			Program->AppendTask(FAnimNextBlendTwoKeyframesPreserveRootMotionTask::Make(WeightedAverageBlendWeight));
+		}
+	}
 
-			for (int32 Index = 0; Index < NumMixerEntries; )
-			{
-				TSharedPtr<FMovieSceneAnimMixerEntry> PeekEntry = Mixer->WeakEntries[Index].Pin();
-				if (!ensure(PeekEntry))
-				{
-					continue;
-				}
+	return Program;
+}
 
-				const int32 Priority = PeekEntry->Priority;
+} // namespace UE::MovieScene
 
-				// Recompute the accumulated weights for this priority group
-				float AccumulatedPriorityWeight = 1.f;
-				// If there is more than one pose in this priority with an absolute weight, we need to blend those together using a weighted average before applying additives and blending with the next priority or the base pose
-				const bool bNeedsSeparateWeightStack = LookaheadComputeAccumulatedWeight(Index, Priority, AccumulatedPriorityWeight);
 
-				bool bIsFirstAbsoluteBlend = true;
-				for ( ; Index < NumMixerEntries; ++Index)
-				{
-					const TSharedPtr<FMovieSceneAnimMixerEntry>& MixerEntry = Mixer->WeakEntries[Index].Pin();
-					if (!MixerEntry)
-					{
-						continue;
-					}
+namespace UE::MovieScene
+{
+	struct FUpdateTaskEntities
+	{
+		static void ForEachEntity(
+			TSharedPtr<FAnimNextEvaluationTask> Task,
+			TSharedPtr<FMovieSceneAnimMixerEntry> MixerEntry,
+			const FMovieSceneRootMotionSettings* RootMotionSettings,
+			const double* WeightAndEasings,
+			const FTransform* BoneMatchTransform)
+		{
+			if (MixerEntry)
+			{
+				MixerEntry->EvalTask = Task;
 
-					if (MixerEntry->Priority != Priority)
-					{
-						// Break and continue with the outer loop
-						break;
-					}
+				// If bSkipPoseWeight is set, use 1.0 for weight
+				MixerEntry->PoseWeight = (MixerEntry->bSkipPoseWeight || !WeightAndEasings) ? 1.0 : *WeightAndEasings;
 
-					if (!MixerEntry->EvalTask.IsValid())
-					{
-						continue;
-					}
+				if (RootMotionSettings)
+				{
+					MixerEntry->RootMotionSettings = *RootMotionSettings;
+				}
 
-					// Skip entries whose source sections are handled by an active transition
-					if (SectionsInTransition.Contains(MixerEntry->EntityOwner))
-					{
-						continue;
-					}
+				if (BoneMatchTransform)
+				{
+					MixerEntry->BoneMatchTransform = *BoneMatchTransform;
+				}
+				else
+				{
+					MixerEntry->BoneMatchTransform.Reset();
+				}
+			}
+		}
+	};
 
-					if (bNeedsSeparateWeightStack && FMath::IsNearlyZero(MixerEntry->PoseWeight, KINDA_SMALL_NUMBER))
-					{
-						continue;
-					}
+	struct FEvaluateAnimMixers
+	{
+		TMap<FMovieSceneAnimMixerKey, TSharedPtr<FMovieSceneAnimMixer>>* Mixers;
+		TMap<FObjectKey, TWeakPtr<FMovieSceneMixerRootMotionComponentData>>* RootMotions;
+		const UMovieSceneRootMotionSystem* RootMotionSystem;
+		const UMovieSceneTransformOriginSystem* TransformOriginSystem;
+		UMovieSceneEntitySystemLinker* Linker;
+		UMovieSceneAnimMixerSystem* MixerSystem;
 
-					// Evaluate the pose itself
-					Mixer->EvaluationProgram->AppendTaskPtr(MixerEntry->EvalTask);
+		FEvaluateAnimMixers(TMap<FMovieSceneAnimMixerKey, TSharedPtr<FMovieSceneAnimMixer>>* InMixers, TMap<FObjectKey, TWeakPtr<FMovieSceneMixerRootMotionComponentData>>* InRootMotions, UMovieSceneEntitySystemLinker* InLinker, UMovieSceneAnimMixerSystem* InMixerSystem)
+			: Mixers(InMixers)
+			, RootMotions(InRootMotions)
+			, Linker(InLinker)
+			, MixerSystem(InMixerSystem)
+		{
+			RootMotionSystem = InLinker->FindSystem<UMovieSceneRootMotionSystem>();
+			TransformOriginSystem = InLinker->FindSystem<UMovieSceneTransformOriginSystem>();
+		}
 
-					// Mirror if specified
-					if (MixerEntry->MirrorTable)
+		// Get an InstanceHandle from the mixer's entries (any valid entry will do).
+		// Mixer entities themselves don't carry InstanceHandle — only section entities do.
+		static FInstanceHandle GetInstanceHandleFromMixer(const FMovieSceneAnimMixer& Mixer)
+		{
+			for (const TWeakPtr<FMovieSceneAnimMixerEntry>& WeakEntry : Mixer.WeakEntries)
+			{
+				if (TSharedPtr<FMovieSceneAnimMixerEntry> Entry = WeakEntry.Pin())
+				{
+					if (Entry->InstanceHandle.IsValid())
 					{
-						// Add a mirroring task
-						TSharedPtr<FAnimNextEvaluationMirroringTask> MirrorTask = MakeShared<FAnimNextEvaluationMirroringTask>();
-						MirrorTask->Setup.MirrorDataTable = MixerEntry->MirrorTable;
-						Mixer->EvaluationProgram->AppendTaskPtr(MirrorTask);
+						return Entry->InstanceHandle;
 					}
+				}
+			}
+			return FInstanceHandle();
+		}
 
-					// Handle root space conversions and manipulations
-					{
-						FTransform TransformOrigin  = FTransform::Identity;
-						FTransform RootTransform    = FTransform::Identity;
-						FVector    RootOffsetOrigin = FVector(ForceInitToZero);
-
-						ESpaceConversions Conversion = ESpaceConversions::None;
-						EMovieSceneRootMotionSpace ThisRootSpace = EMovieSceneRootMotionSpace::AnimationSpace;
-
-						if (MixerEntry->RootMotionSettings)
-						{
-							ThisRootSpace = MixerEntry->RootMotionSettings->RootMotionSpace;
+		void ForEachEntity(
+			FObjectComponent MeshComponent,
+			const TInstancedStruct<FMovieSceneMixedAnimationTarget>& Target,
+			TSharedPtr<FAnimNextEvaluationTask>& MixerTask
+			) const
+		{
+			TSharedPtr<FMovieSceneAnimMixer> Mixer = Mixers->FindRef(FMovieSceneAnimMixerKey({ MeshComponent.GetObject(), Target }));
+			if (!Mixer)
+			{
+				return;
+			}
 
-							// Handle root transform offset/override
-							const bool bHasLocation = !MixerEntry->RootMotionSettings->RootLocation.IsZero();
-							const bool bHasRotation = !MixerEntry->RootMotionSettings->RootRotation.IsIdentity();
+			USceneComponent* Component = Cast<USceneComponent>(MeshComponent.GetObject());
+			TSharedPtr<FMovieSceneMixerRootMotionComponentData> RootMotion = RootMotions->FindRef(MeshComponent.GetObject()).Pin();
 
-							if (MixerEntry->RootMotionSettings->TransformMode == EMovieSceneRootMotionTransformMode::Override)
-							{
-								// Override mode: Always apply, even with identity values (to null out root motion)
-								Conversion |= ESpaceConversions::RootTransformOverride;
-								RootTransform = FTransform(
-									MixerEntry->RootMotionSettings->RootRotation,
-									MixerEntry->RootMotionSettings->RootLocation
-									);
-							}
-							else if (MixerEntry->RootMotionSettings->TransformMode == EMovieSceneRootMotionTransformMode::Offset)
-							{
-								// Offset mode: Only apply if we have non-zero values to offset by
-								if (bHasLocation || bHasRotation)
-								{
-									Conversion |= ESpaceConversions::RootTransformOffset;
-									RootOffsetOrigin = MixerEntry->RootMotionSettings->RootOriginLocation;
-									RootTransform = FTransform(
-										MixerEntry->RootMotionSettings->RootRotation,
-										MixerEntry->RootMotionSettings->RootLocation
-										);
-								}
-							}
-							// Asset mode: No override/offset applied, use the animation's root motion as-is
-						}
+			FProgramBuildContext BuildContext;
+			BuildContext.RootMotion = RootMotion;
+			BuildContext.RootMotionSystem = RootMotionSystem;
+			BuildContext.TransformOriginSystem = TransformOriginSystem;
+			BuildContext.Linker = Linker;
+			BuildContext.Component = Component;
 
-						if (ThisRootSpace == EMovieSceneRootMotionSpace::AnimationSpace)
-						{
-							Conversion |= ESpaceConversions::ComponentToActorRotation;
+			// Apply bake filter from the bake target if present.
+			FProgramBuildFilter BuildFilter;
+			if (const FMovieSceneMixerBakeTarget* BakeTarget = Target.GetPtr<FMovieSceneMixerBakeTarget>())
+			{
+				if (!BakeTarget->Filter.ExcludeSections.IsEmpty())
+				{
+					BuildFilter.ExcludeSections = &BakeTarget->Filter.ExcludeSections;
+				}
+				if (!BakeTarget->Filter.IncludeOnlySections.IsEmpty())
+				{
+					BuildFilter.IncludeOnly = &BakeTarget->Filter.IncludeOnlySections;
+				}
+				BuildFilter.MaxPriority = BakeTarget->Filter.MaxPriority;
+				BuildFilter.bSkipRootMotionConversion = BakeTarget->Filter.bSkipRootMotionConversion;
+				BuildFilter.CaptureAnimSpaceRoot = BakeTarget->Filter.CaptureAnimSpaceRoot;
+			}
 
-							if (bRootComponentHasKeyedTransform)
-							{
-								Conversion |= ESpaceConversions::AnimationToWorld;
-							}
-							else if (TransformOriginSystem && TransformOriginSystem->GetTransformOrigin(MixerEntry->InstanceHandle, TransformOrigin))
-							{
-								Conversion |= ESpaceConversions::TransformOriginToWorld;
-							}
-						}
-						else
-						{
-							Conversion |= ESpaceConversions::WorldSpaceComponentTransformCompensation;
-						}
+			// Build program from entries. Gap entries (null EvalTask) are naturally
+			// skipped, so this returns null during a full gap.
+			Mixer->EvaluationProgram = BuildProgramFromEntries(*Mixer, BuildContext, BuildFilter);
 
-						if (Conversion != ESpaceConversions::None)
+			if (RootMotion)
+			{
+				// Check if any entry produces root motion. If none do, this is a
+				// gap (either full — no real entries at all, or soft — only non-root-motion
+				// sections like facial are active). In either case, apply the persisted
+				// gap transform if configured.
+				bool bAnyEntryHasRootMotion = false;
+				for (const TWeakPtr<FMovieSceneAnimMixerEntry>& WeakEntry : Mixer->WeakEntries)
+				{
+					if (TSharedPtr<FMovieSceneAnimMixerEntry> Entry = WeakEntry.Pin())
+					{
+						if (Entry->RootMotionSettings.IsSet())
 						{
-							// Need to convert root motion attribute to world space for blending, insert task to do that
-							Mixer->EvaluationProgram->AppendTask(FAnimNextConvertRootMotionToWorldSpaceTask::Make(
-								RootMotion, TransformOrigin, RootTransform, RootOffsetOrigin, Conversion));
+							bAnyEntryHasRootMotion = true;
+							break;
 						}
 					}
+				}
 
-					if (MixerEntry->bRequiresBlend)
+				if (!bAnyEntryHasRootMotion
+					&& RootMotion->GapBehavior == EMovieSceneRootMotionGapBehavior::PersistPreviousTransform)
+				{
+					UMovieSceneRootMotionSection* Section = RootMotion->WeakRootMotionSection.Get();
+					FInstanceHandle EntryInstanceHandle = GetInstanceHandleFromMixer(*Mixer);
+					if (Section && EntryInstanceHandle.IsValid())
 					{
-						if (MixerEntry->bAdditive)
+						const FMovieSceneContext& Context = Linker->GetInstanceRegistry()->GetContext(EntryInstanceHandle);
+						FTransform GapTransform = Section->GetAccumulatedOffsetAt(Context.GetTime());
+						if (!GapTransform.Equals(FTransform::Identity))
 						{
-							// Add additive blend task.
-							Mixer->EvaluationProgram->AppendTask(FAnimNextApplyAdditiveKeyframeTask::Make(MixerEntry->PoseWeight));
-						}
-						else if (bNeedsSeparateWeightStack)
-						{
-							// Add absolute blend task. The first absolute task within a priority level needs to overwrite the pose with its weighted average,
-							//     subsequent poses are added to this. This in effect leaves us with a single pose on the stack at the end of this priority level
-							//     that includes the weighted average of all absolutes within this priority level, plus all absolutes
-							if (bIsFirstAbsoluteBlend)
-							{
-								bIsFirstAbsoluteBlend = false;
-								Mixer->EvaluationProgram->AppendTask(FAnimNextBlendOverwriteKeyframeWithScaleTask::Make(MixerEntry->PoseWeight / AccumulatedPriorityWeight));
-							}
-							else
+							if (!Mixer->EvaluationProgram)
 							{
-								Mixer->EvaluationProgram->AppendTask(FMovieSceneAccumulateAbsoluteBlendTask::Make(MixerEntry->PoseWeight / AccumulatedPriorityWeight));
+								Mixer->EvaluationProgram = MakeShared<UE::UAF::FEvaluationProgram>();
 							}
-						}
-						else
-						{
-							Mixer->EvaluationProgram->AppendTask(FAnimNextBlendTwoKeyframesPreserveRootMotionTask::Make(MixerEntry->PoseWeight));
+							// Set the cached gap transform directly as the root motion
+							// attribute. The value is already in post-conversion world
+							// space from the trajectory cache.
+							Mixer->EvaluationProgram->AppendTask(
+								FAnimNextSetRootTransformAttributeTask::Make(GapTransform));
 						}
 					}
 				}
-				// If we have a previous amount to blend with the last priority, add the tasks to ensure that happens
-				if (bNeedsSeparateWeightStack)
-				{
-					Mixer->EvaluationProgram->AppendTask(FAnimNextNormalizeKeyframeRotationsTask());
 
-					// For now we always blend weighted averages with the next pose with a weight of 1.
-					//    Ultimately this should be controlled by its own weight within an anim mixer layer track.
-					constexpr float WeightedAverageBlendWeight = 1.f;
-					Mixer->EvaluationProgram->AppendTask(FAnimNextBlendTwoKeyframesPreserveRootMotionTask::Make(WeightedAverageBlendWeight));
+				if (Mixer->EvaluationProgram)
+				{
+					// Skip the store task during bake — bake captures the root
+					// motion attribute directly from the result keyframe.
+					if (!Target.GetPtr<FMovieSceneMixerBakeTarget>())
+					{
+						USceneComponent* Root = Component ? Component->GetOwner()->GetRootComponent() : nullptr;
+						const bool bComponentHasKeyedTransform = (Component && RootMotionSystem->IsTransformKeyed(Component));
+						const bool bRootComponentHasKeyedTransform = (Root && RootMotionSystem->IsTransformKeyed(Root));
+						Mixer->EvaluationProgram->AppendTask(FAnimNextStoreRootTransformTask::Make(RootMotion, bComponentHasKeyedTransform, bRootComponentHasKeyedTransform));
+					}
 				}
 			}
 
-			if (RootMotion)
+			if (Mixer->EvaluationProgram)
+			{
+				if (!MixerTask.IsValid())
+				{
+					MixerTask = MakeShared<FAnimNextExecuteProgramTask>();
+				}
+				TSharedPtr<FAnimNextExecuteProgramTask> EvalProgramTask = StaticCastSharedPtr<FAnimNextExecuteProgramTask>(MixerTask);
+				EvalProgramTask->Program = Mixer->EvaluationProgram;
+			}
+			else
 			{
-				// Add task to store the final root motion result
-				Mixer->EvaluationProgram->AppendTask(FAnimNextStoreRootTransformTask::Make(RootMotion, bComponentHasKeyedTransform, bRootComponentHasKeyedTransform));
+				MixerTask = nullptr;
 			}
 		}
 	};
@@ -605,7 +942,7 @@ void FMovieSceneAccumulateAbsoluteBlendTask::Execute(UE::UAF::FEvaluationVM& VM)
 	if (!VM.PopValue(KEYFRAME_STACK_NAME, KeyframeB))
 	{
 		// We have a single input, leave it on top of the stack
-		VM.PushValue(KEYFRAME_STACK_NAME, MoveTemp(KeyframeB));
+		VM.PushValue(KEYFRAME_STACK_NAME, MoveTemp(KeyframeA));
 		return;
 	}
 
@@ -739,6 +1076,42 @@ void UMovieSceneAnimMixerSystem::OnLink()
 {
 	RootMotionSystem = Linker->LinkSystem<UMovieSceneRootMotionSystem>();
 	Linker->SystemGraph.AddReference(this, RootMotionSystem);
+
+	PreFlushHandle = Linker->Events.PreFlush.AddUObject(this, &UMovieSceneAnimMixerSystem::HandlePreFlush);
+}
+
+void UMovieSceneAnimMixerSystem::OnUnlink()
+{
+	if (PreFlushHandle.IsValid() && Linker)
+	{
+		Linker->Events.PreFlush.Remove(PreFlushHandle);
+		PreFlushHandle.Reset();
+	}
+}
+
+void UMovieSceneAnimMixerSystem::HandlePreFlush(UMovieSceneEntitySystemLinker* InLinker)
+{
+	for (TPair<FMovieSceneAnimMixerKey, TSharedPtr<FMovieSceneAnimMixer>>& Pair : Mixers)
+	{
+		UMovieSceneAnimationMixerTrack* Track = Pair.Value->OwnerTrack.Get();
+		if (!Track || !Track->HasDirtyAccumulatedOffsetCache())
+		{
+			continue;
+		}
+
+		// Find a valid instance handle from any live entry in this mixer
+		for (const TWeakPtr<FMovieSceneAnimMixerEntry>& WeakEntry : Pair.Value->WeakEntries)
+		{
+			if (TSharedPtr<FMovieSceneAnimMixerEntry> Entry = WeakEntry.Pin())
+			{
+				if (Entry->InstanceHandle.IsValid())
+				{
+					Track->RebuildDirtyAccumulatedOffsetCache(InLinker, Entry->InstanceHandle);
+					break;
+				}
+			}
+		}
+	}
 }
 
 void UMovieSceneAnimMixerSystem::OnRun(FSystemTaskPrerequisites& InPrerequisites, FSystemSubsequentTasks& Subsequents)
@@ -896,6 +1269,14 @@ void UMovieSceneAnimMixerSystem::OnRun(FSystemTaskPrerequisites& InPrerequisites
 				}
 
 				Mixer->bNeedsResort = true;
+
+				if (!Mixer->OwnerTrack.IsValid())
+				{
+					if (UMovieSceneSection* Section = Cast<UMovieSceneSection>(Entry->EntityOwner.ResolveObjectPtr()))
+					{
+						Mixer->OwnerTrack = Section->GetTypedOuter<UMovieSceneAnimationMixerTrack>();
+					}
+				}
 			});
 
 		// ----------------------------------------------------------------------------------------------------------------------------------------
@@ -1055,6 +1436,7 @@ void UMovieSceneAnimMixerSystem::OnSchedulePersistentTasks(UE::MovieScene::IEnti
 		.Write(AnimMixerComponents->MixerEntry)
 		.ReadOptional(AnimMixerComponents->RootMotionSettings)
 		.ReadOptional(BuiltInComponents->WeightAndEasingResult)
+		.ReadOptional(AnimMixerComponents->BoneMatchTransform)
 		.FilterNone({BuiltInComponents->Tags.NeedsUnlink, BuiltInComponents->Tags.Ignored})
 		.Schedule_PerEntity<FUpdateTaskEntities>(&Linker->EntityManager, TaskScheduler);
 
@@ -1065,7 +1447,7 @@ void UMovieSceneAnimMixerSystem::OnSchedulePersistentTasks(UE::MovieScene::IEnti
 		.Read(AnimMixerComponents->Target)
 		.Write(AnimMixerComponents->MixerTask)
 		.FilterNone({BuiltInComponents->Tags.Ignored})
-		.Schedule_PerEntity<FEvaluateAnimMixers>(&Linker->EntityManager, TaskScheduler, &Mixers, &RootMotionByObject, Linker);
+		.Schedule_PerEntity<FEvaluateAnimMixers>(&Linker->EntityManager, TaskScheduler, &Mixers, &RootMotionByObject, Linker, this);
 
 	TaskScheduler->AddPrerequisite(UpdateTask, MixTask);
 }
diff --git a/Engine/Plugins/Experimental/MovieSceneAnimMixer/Source/MovieSceneAnimMixer/Private/Systems/MovieSceneMixerBakeTargetSystem.cpp b/Engine/Plugins/Experimental/MovieSceneAnimMixer/Source/MovieSceneAnimMixer/Private/Systems/MovieSceneMixerBakeTargetSystem.cpp
new file mode 100644
--- /dev/null
+++ b/Engine/Plugins/Experimental/MovieSceneAnimMixer/Source/MovieSceneAnimMixer/Private/Systems/MovieSceneMixerBakeTargetSystem.cpp
@@ -0,0 +1,143 @@
+// Copyright Epic Games, Inc. All Rights Reserved.
+
+#include "Systems/MovieSceneMixerBakeTargetSystem.h"
+
+#include "EntitySystem/Interrogation/MovieSceneInterrogationLinker.h"
+#include "EntitySystem/MovieSceneEntitySystemLinker.h"
+#include "EntitySystem/MovieSceneEntitySystemTask.h"
+#include "MovieSceneTracksComponentTypes.h"
+#include "AnimMixerComponentTypes.h"
+#include "Systems/MovieSceneAnimMixerSystem.h"
+#include "Components/SkeletalMeshComponent.h"
+
+#include "Animation/BuiltInAttributeTypes.h"
+#include "Misc/MemStack.h"
+#include "EvaluationVM/EvaluationTask.h"
+#include "EvaluationVM/EvaluationVM.h"
+#include "GenerationTools.h"
+#include "Graph/AnimNext_LODPose.h"
+#include "BonePose.h"
+
+#include UE_INLINE_GENERATED_CPP_BY_NAME(MovieSceneMixerBakeTargetSystem)
+
+namespace UE::MovieScene
+{
+
+// Task that executes mixer programs in a temporary VM, capturing the
+// resulting keyframe and root motion transform instead of applying them.
+struct FCaptureBakeResults
+{
+	UMovieSceneEntitySystemLinker* Linker;
+	UMovieSceneMixerBakeTargetSystem* BakeSystem;
+
+	FCaptureBakeResults(UMovieSceneEntitySystemLinker* InLinker, UMovieSceneMixerBakeTargetSystem* InBakeSystem)
+		: Linker(InLinker)
+		, BakeSystem(InBakeSystem)
+	{
+	}
+
+	void ForEachEntity(
+		FObjectComponent MeshComponent,
+		const TInstancedStruct<FMovieSceneMixedAnimationTarget>& Target,
+		const TSharedPtr<FAnimNextEvaluationTask>& MixerTask) const
+	{
+		if (!BakeSystem->bBakeActive || !MixerTask.IsValid())
+		{
+			return;
+		}
+
+		// Only process entities targeted with FMovieSceneMixerBakeTarget.
+		// During bake evaluation the track's target is temporarily swapped
+		// to this type, so only bake-targeted entities end up here.
+		if (!Target.GetPtr<FMovieSceneMixerBakeTarget>())
+		{
+			return;
+		}
+
+		USkeletalMeshComponent* SkelMeshComp = Cast<USkeletalMeshComponent>(MeshComponent.GetObject());
+		if (!SkelMeshComp || !SkelMeshComp->GetSkeletalMeshAsset())
+		{
+			return;
+		}
+
+		using namespace UE::UAF;
+
+		FMemMark Mark(FMemStack::Get());
+
+		// Set up a temporary VM with the skeletal mesh's reference pose
+		FDataHandle RefPoseHandle = FDataRegistry::Get()->GetOrGenerateReferencePose(SkelMeshComp);
+		const FReferencePose& RefPose = RefPoseHandle.GetRef<FReferencePose>();
+		int32 LODLevel = SkelMeshComp->GetPredictedLODLevel();
+
+		FEvaluationVM TempVM(EEvaluationFlags::All, RefPose, LODLevel);
+
+		// Push a reference keyframe as the initial pose
+		FKeyframeState InitialKeyframe = TempVM.MakeReferenceKeyframe(false);
+		TempVM.PushValue(KEYFRAME_STACK_NAME, MakeUnique<FKeyframeState>(MoveTemp(InitialKeyframe)));
+
+		// Execute the mixer task
+		MixerTask->Execute(TempVM);
+
+		// Pop the resulting keyframe. The store root transform task is
+		// deliberately skipped during bake (see FEvaluateAnimMixers), so the
+		// RootTransformAttribute is still present on the keyframe with the
+		// world-space root motion value from the convert task.
+		UMovieSceneMixerBakeTargetSystem::FBakedResult Result;
+
+		TUniquePtr<FKeyframeState> ResultKeyframe;
+		if (TempVM.PopValue(KEYFRAME_STACK_NAME, ResultKeyframe))
+		{
+			// Read root motion from the RootTransformAttribute on the pose
+			const FAnimMixerComponentTypes* AnimMixerComponents = FAnimMixerComponentTypes::Get();
+			const FTransformAnimationAttribute* RootAttr =
+				ResultKeyframe->Attributes.Find<FTransformAnimationAttribute>(
+					AnimMixerComponents->RootTransformAttributeId);
+			if (RootAttr)
+			{
+				Result.RootMotionTransform = RootAttr->Value;
+			}
+
+			Result.Keyframe = MoveTemp(ResultKeyframe);
+		}
+
+		BakeSystem->BakedResults.Add(MoveTemp(Result));
+	}
+};
+
+} // namespace UE::MovieScene
+
+UMovieSceneMixerBakeTargetSystem::UMovieSceneMixerBakeTargetSystem(const FObjectInitializer& Init)
+	: Super(Init)
+{
+	using namespace UE::MovieScene;
+
+	FAnimMixerComponentTypes* AnimMixerComponents = FAnimMixerComponentTypes::Get();
+
+	RelevantComponent = AnimMixerComponents->MixerTask;
+	Phase = ESystemPhase::Scheduling;
+
+	// Excluded from interrogation
+	SystemCategories |= FSystemInterrogator::GetExcludedFromInterrogationCategory();
+
+	if (HasAnyFlags(RF_ClassDefaultObject))
+	{
+		DefineImplicitPrerequisite(UMovieSceneAnimMixerSystem::StaticClass(), GetClass());
+	}
+}
+
+void UMovieSceneMixerBakeTargetSystem::OnSchedulePersistentTasks(UE::MovieScene::IEntitySystemScheduler* TaskScheduler)
+{
+	using namespace UE::MovieScene;
+
+	FAnimMixerComponentTypes* AnimMixerComponents = FAnimMixerComponentTypes::Get();
+
+	FTaskParams Params(TEXT("Capture Bake Results"));
+	Params.ForceGameThread();
+	FEntityTaskBuilder()
+		.Read(AnimMixerComponents->MeshComponent)
+		.Read(AnimMixerComponents->Target)
+		.Read(AnimMixerComponents->MixerTask)
+		.SetParams(Params)
+		.Schedule_PerEntity<FCaptureBakeResults>(&Linker->EntityManager, TaskScheduler,
+			Linker, this);
+}
diff --git a/Engine/Plugins/Experimental/MovieSceneAnimMixer/Source/MovieSceneAnimMixer/Private/Systems/MovieSceneRootMotionSystem.cpp b/Engine/Plugins/Experimental/MovieSceneAnimMixer/Source/MovieSceneAnimMixer/Private/Systems/MovieSceneRootMotionSystem.cpp
--- a/Engine/Plugins/Experimental/MovieSceneAnimMixer/Source/MovieSceneAnimMixer/Private/Systems/MovieSceneRootMotionSystem.cpp
+++ b/Engine/Plugins/Experimental/MovieSceneAnimMixer/Source/MovieSceneAnimMixer/Private/Systems/MovieSceneRootMotionSystem.cpp
@@ -4,6 +4,7 @@
 
 #include "EntitySystem/BuiltInComponentTypes.h"
 #include "EntitySystem/IMovieSceneTaskScheduler.h"
+#include "EntitySystem/MovieSceneBoundObjectInstantiator.h"
 #include "MovieSceneTracksComponentTypes.h"
 #include "MovieSceneRootMotionSection.h"
 
@@ -33,6 +34,7 @@
 #include "PreAnimatedState/MovieScenePreAnimatedComponentTransformStorage.h"
 #include "EntitySystem/Interrogation/MovieSceneInterrogatedPropertyInstantiator.h"
 #include "EntitySystem/MovieSceneEntitySystemRunner.h"
+#include "EntitySystem/MovieSceneInstanceRegistry.h"
 #include "Misc/ScopeRWLock.h"
 #include "VisualLogger/VisualLogger.h"
 
@@ -48,16 +50,31 @@ namespace UE::MovieScene
 {
 	struct FGatherRootDestinations
 	{
-		static void ForEachEntity(
+		UMovieSceneAnimMixerSystem* AnimMixer;
+
+		FGatherRootDestinations(UMovieSceneAnimMixerSystem* InAnimMixer)
+			: AnimMixer(InAnimMixer)
+		{}
+
+		void ForEachEntity(
+			FObjectKey BoundObjectKey,
 			uint8 ByteResult,
 			EMovieSceneRootMotionDestination& OutRootDestination,
-			TSharedPtr<FMovieSceneMixerRootMotionComponentData>& RootMotion
-			)
+			TSharedPtr<FMovieSceneMixerRootMotionComponentData>& OutRootMotion
+			) const
 		{
 			OutRootDestination = (EMovieSceneRootMotionDestination)(ByteResult);
-			if (ensure(RootMotion))
+
+			// Fall back to the mixer system's root motion map if the entity's
+			// component is null (can happen with default-constructed child entities).
+			if (!OutRootMotion)
+			{
+				OutRootMotion = AnimMixer->FindRootMotion(BoundObjectKey);
+			}
+
+			if (OutRootMotion)
 			{
-				RootMotion->RootDestination = OutRootDestination;
+				OutRootMotion->RootDestination = OutRootDestination;
 			}
 		}
 	};
@@ -299,41 +316,62 @@ void FAnimNextConvertRootMotionToWorldSpaceTask::Execute(UE::UAF::FEvaluationVM&
 		? RootTransform
 		: RootMotionAttribute->Value;
 
+	if (CaptureAnimSpaceRoot)
+	{
+		*CaptureAnimSpaceRoot = RelativeTransform;
+	}
+
 	TSharedPtr<FMovieSceneMixerRootMotionComponentData> RootMotionData = WeakRootMotionData.Pin();
-	if (RootMotionData)
+
+	// AccumulatedOffset: compute the animation-space delta from the section's
+	// initial root transform and re-base it onto the accumulated position.
+	// Runs first while RelativeTransform is still in raw animation space.
+	if (EnumHasAnyFlags(Conversions, ESpaceConversions::AccumulatedOffset))
+	{
+		FTransform Delta = InitialAnimRootTransform.Inverse() * RelativeTransform;
+		RelativeTransform = Delta * AccumulatedTransform;
+	}
+	else
 	{
-		if (EnumHasAnyFlags(Conversions, ESpaceConversions::ComponentToActorRotation))
+		if (RootMotionData)
 		{
-			TOptional<FQuat> InverseMeshToActorRotation = RootMotionData->GetInverseMeshToActorRotation();
-			if (InverseMeshToActorRotation.IsSet())
+			if (EnumHasAnyFlags(Conversions, ESpaceConversions::ComponentToActorRotation))
 			{
-				RelativeTransform.SetTranslation(InverseMeshToActorRotation.GetValue() * RelativeTransform.GetTranslation());
+				TOptional<FQuat> InverseMeshToActorRotation = RootMotionData->GetInverseMeshToActorRotation();
+				if (InverseMeshToActorRotation.IsSet())
+				{
+					RelativeTransform.SetTranslation(InverseMeshToActorRotation.GetValue() * RelativeTransform.GetTranslation());
+				}
 			}
 		}
-	}
 
-	// Apply root transform offsets if necessary
-	if (EnumHasAnyFlags(Conversions, ESpaceConversions::RootTransformOffset))
-	{
-		// Multiply by the inverse of the root origin to get the root motion relative to the origin
-		RelativeTransform *= FTransform(-RootOffsetOrigin);
-		RelativeTransform = RelativeTransform * RootTransform;
-		// Put the transform back in component space
-		RelativeTransform *= FTransform(RootOffsetOrigin);
-	}
+		// Apply bone match offset in animation space before world conversion.
+		// Shifts the root motion so the matched bone aligns with the underlying pose.
+		if (EnumHasAnyFlags(Conversions, ESpaceConversions::BoneMatchOffset))
+		{
+			RelativeTransform = RelativeTransform * BoneMatchTransform;
+		}
 
-	if (RootMotionData)
-	{
-		if (EnumHasAnyFlags(Conversions, ESpaceConversions::AnimationToWorld))
+		if (EnumHasAnyFlags(Conversions, ESpaceConversions::RootTransformOffset))
 		{
-			RelativeTransform *= RootMotionData->ActorTransform;
+			// Pivot the offset around the bone-matched root position when active,
+			// otherwise around the configured RootOffsetOrigin.
+			FVector Pivot = EnumHasAnyFlags(Conversions, ESpaceConversions::BoneMatchOffset)
+				? RelativeTransform.GetTranslation()
+				: RootOffsetOrigin;
+			RelativeTransform *= FTransform(-Pivot);
+			RelativeTransform = RelativeTransform * RootTransform;
+			RelativeTransform *= FTransform(Pivot);
 		}
-	}
 
-	// Doesn't really make sense to mix both AnimationToWorld & TransformOriginToWorld
-	if (EnumHasAnyFlags(Conversions, ESpaceConversions::TransformOriginToWorld))
-	{
-		RelativeTransform *= TransformOrigin;
+		if (RootMotionData && EnumHasAnyFlags(Conversions, ESpaceConversions::AnimationToWorld))
+		{
+			RelativeTransform *= RootMotionData->ActorTransform;
+		}
+		else if (EnumHasAnyFlags(Conversions, ESpaceConversions::TransformOriginToWorld))
+		{
+			RelativeTransform *= TransformOrigin;
+		}
 	}
 
 	if (EnumHasAnyFlags(Conversions, ESpaceConversions::WorldSpaceComponentTransformCompensation) && RootMotionData)
@@ -507,6 +545,37 @@ void FAnimNextStoreRootTransformTask::Execute(UE::UAF::FEvaluationVM& VM) const
 	}
 }
 
+FAnimNextSetRootTransformAttributeTask FAnimNextSetRootTransformAttributeTask::Make(const FTransform& InTransform)
+{
+	FAnimNextSetRootTransformAttributeTask Task;
+	Task.Transform = InTransform;
+	return Task;
+}
+
+void FAnimNextSetRootTransformAttributeTask::Execute(UE::UAF::FEvaluationVM& VM) const
+{
+	using namespace UE::UAF;
+	using namespace UE::MovieScene;
+
+	if (!EnumHasAnyFlags(VM.GetFlags(), EEvaluationFlags::Attributes | EEvaluationFlags::Trajectory))
+	{
+		return;
+	}
+
+	TUniquePtr<FKeyframeState>* Keyframe = VM.PeekValueMutable<TUniquePtr<FKeyframeState>>(KEYFRAME_STACK_NAME, 0);
+	if (!Keyframe || !Keyframe->IsValid())
+	{
+		return;
+	}
+
+	const FAnimMixerComponentTypes* AnimMixerComponents = FAnimMixerComponentTypes::Get();
+
+	// FindOrAdd creates the attribute if it doesn't already exist on the pose
+	FTransformAnimationAttribute* RootMotionAttribute =
+		(*Keyframe)->Attributes.FindOrAdd<FTransformAnimationAttribute>(AnimMixerComponents->RootTransformAttributeId);
+	RootMotionAttribute->Value = Transform;
+}
+
 UMovieSceneRootMotionSystem::UMovieSceneRootMotionSystem(const FObjectInitializer& ObjInit)
 	: Super(ObjInit)
 {
@@ -524,6 +593,10 @@ UMovieSceneRootMotionSystem::UMovieSceneRootMotionSystem(const FObjectInitialize
 		// This must be run before the anim mixer to ensure that the anim mixer sets up its program correctly with root motion
 		DefineImplicitPrerequisite(GetClass(), UMovieSceneAnimMixerSystem::StaticClass());
 
+		// Must run after bound object instantiation so that bound child entities
+		// (which carry RootDestination + MixerRootMotion) exist when we process them
+		DefineImplicitPrerequisite(UMovieSceneGenericBoundObjectInstantiator::StaticClass(), GetClass());
+
 		DefineImplicitPrerequisite(UByteChannelEvaluatorSystem::StaticClass(), GetClass());
 		DefineImplicitPrerequisite(UDoubleChannelEvaluatorSystem::StaticClass(), GetClass());
 		DefineImplicitPrerequisite(UMovieScenePiecewiseDoubleBlenderSystem::StaticClass(), GetClass());
@@ -579,6 +652,8 @@ void UMovieSceneRootMotionSystem::OnRun(FSystemTaskPrerequisites& InPrerequisite
 		.Read(BuiltInComponents->RootInstanceHandle)
 		.Read(BuiltInComponents->BoundObjectKey)
 		.Read(BuiltInComponents->BoundObject)
+		.ReadOptional(AnimMixerComponents->GapBehavior)
+		.ReadOptional(AnimMixerComponents->EntityOwner)
 		.Write(AnimMixerComponents->Target)
 		.Write(AnimMixerComponents->MixerRootMotion)
 		.FilterAll({ AnimMixerComponents->RootDestination, BuiltInComponents->Tags.NeedsLink })
@@ -589,6 +664,8 @@ void UMovieSceneRootMotionSystem::OnRun(FSystemTaskPrerequisites& InPrerequisite
 				TRead<FRootInstanceHandle> RootInstanceHandles,
 				TRead<FObjectKey> BoundObjectKeys,
 				TRead<UObject*> BoundObjects,
+				TReadOptional<EMovieSceneRootMotionGapBehavior> GapBehaviors,
+				TReadOptional<FObjectKey> EntityOwners,
 				TWrite<TInstancedStruct<FMovieSceneMixedAnimationTarget>> OutTargets,
 				TWrite<TSharedPtr<FMovieSceneMixerRootMotionComponentData>> OutRootMotions)
 			{
@@ -663,6 +740,17 @@ void UMovieSceneRootMotionSystem::OnRun(FSystemTaskPrerequisites& InPrerequisite
 
 					USceneComponent* Target = Cast<USceneComponent>(BoundObjectKey.ResolveObjectPtr());
 					OutRootMotion->OriginalBoundObject = Target;
+
+					// Store gap behavior data from the root motion section entity
+					if (GapBehaviors.IsValid())
+					{
+						OutRootMotion->GapBehavior = GapBehaviors[Index];
+					}
+					if (EntityOwners.IsValid())
+					{
+						UMovieSceneRootMotionSection* Section = Cast<UMovieSceneRootMotionSection>(EntityOwners[Index].ResolveObjectPtr());
+						OutRootMotion->WeakRootMotionSection = Section;
+					}
 				}
 			}
 		);
@@ -710,11 +798,12 @@ void UMovieSceneRootMotionSystem::OnSchedulePersistentTasks(UE::MovieScene::IEnt
 
 	// Gather root destination results
 	FTaskID GatherRootDestinationTask = FEntityTaskBuilder()
+	.Read(BuiltInComponents->BoundObjectKey)
 	.Read(BuiltInComponents->ByteResult)
 	.Write(AnimMixerComponents->RootDestination)
 	.Write(AnimMixerComponents->MixerRootMotion)
 	.FilterNone({ BuiltInComponents->Tags.Ignored })
-	.Schedule_PerEntity<FGatherRootDestinations>(&Linker->EntityManager, TaskScheduler);
+	.Schedule_PerEntity<FGatherRootDestinations>(&Linker->EntityManager, TaskScheduler, AnimMixer);
 
 	// Reset root motion data and gather current component/actor transforms
 	FTaskParams InitialzeParams(TEXT("Initialize Root Motion"));
diff --git a/Engine/Plugins/Experimental/MovieSceneAnimMixer/Source/MovieSceneAnimMixer/Public/AnimMixerBakeEvaluation.h b/Engine/Plugins/Experimental/MovieSceneAnimMixer/Source/MovieSceneAnimMixer/Public/AnimMixerBakeEvaluation.h
new file mode 100644
--- /dev/null
+++ b/Engine/Plugins/Experimental/MovieSceneAnimMixer/Source/MovieSceneAnimMixer/Public/AnimMixerBakeEvaluation.h
@@ -0,0 +1,72 @@
+// Copyright Epic Games, Inc. All Rights Reserved.
+
+#pragma once
+
+#include "CoreTypes.h"
+#include "Misc/FrameTime.h"
+#include "UObject/ObjectKey.h"
+
+class UMovieSceneEntitySystemLinker;
+class UMovieSceneAnimationMixerTrack;
+class USkeletalMeshComponent;
+
+namespace UE::MovieScene { struct FInstanceHandle; }
+
+#include "EvaluationVM/KeyframeState.h"
+
+namespace UE::MovieScene::AnimMixerBakeEvaluation
+{
+
+// Filter for controlling which mixer entries are included during bake evaluation.
+struct FBakeFilter
+{
+	// If non-empty, only include entries whose EntityOwner matches
+	TSet<FObjectKey> IncludeOnlySections;
+
+	// Entries whose EntityOwner matches are excluded
+	TSet<FObjectKey> ExcludeSections;
+
+	// If >= 0, only entries with priority <= this value are included
+	int32 MaxPriority = -1;
+
+	// When true, the ConvertRootMotionToWorldSpaceTask is omitted from the
+	// evaluation program. The resulting RootMotionTransform will be in
+	// animation space rather than world space.
+	bool bSkipRootMotionConversion = false;
+
+	// When set, the conversion task writes the pre-conversion animation-space
+	// root here. Populated automatically by FBakeEvaluationScope.
+	TSharedPtr<FTransform> CaptureAnimSpaceRoot;
+};
+
+struct FBakeResult
+{
+	FFrameTime Time;
+	FTransform RootMotionTransform = FTransform::Identity;
+	FTransform AnimationSpaceRootMotion = FTransform::Identity;
+	TUniquePtr<UE::UAF::FKeyframeState> Keyframe;
+};
+
+// Evaluate the mixer at a single time, capturing results through the bake target system.
+// Temporarily swaps the track's target, sets a bake filter on the mixer system,
+// then re-evaluates the ECS and restores original state.
+MOVIESCENEANIMMIXER_API FBakeResult EvaluateAtTime(
+	UMovieSceneEntitySystemLinker* Linker,
+	FInstanceHandle InstanceHandle,
+	UMovieSceneAnimationMixerTrack* MixerTrack,
+	FFrameTime Time,
+	const FBakeFilter& Filter = {});
+
+// Evaluate the mixer at regular intervals from StartTime (inclusive) for NumSamples.
+// More efficient than calling EvaluateAtTime in a loop because save/restore only
+// happens once around the entire batch.
+TArray<FBakeResult> EvaluateRange(
+	UMovieSceneEntitySystemLinker* Linker,
+	FInstanceHandle InstanceHandle,
+	UMovieSceneAnimationMixerTrack* MixerTrack,
+	FFrameTime StartTime,
+	FFrameTime FrameStep,
+	int32 NumSamples,
+	const FBakeFilter& Filter = {});
+
+} // namespace UE::MovieScene::AnimMixerBakeEvaluation
diff --git a/Engine/Plugins/Experimental/MovieSceneAnimMixer/Source/MovieSceneAnimMixer/Public/AnimMixerBoneMatching.h b/Engine/Plugins/Experimental/MovieSceneAnimMixer/Source/MovieSceneAnimMixer/Public/AnimMixerBoneMatching.h
new file mode 100644
--- /dev/null
+++ b/Engine/Plugins/Experimental/MovieSceneAnimMixer/Source/MovieSceneAnimMixer/Public/AnimMixerBoneMatching.h
@@ -0,0 +1,60 @@
+// Copyright Epic Games, Inc. All Rights Reserved.
+
+#pragma once
+
+#include "CoreTypes.h"
+#include "Containers/Set.h"
+#include "Misc/FrameTime.h"
+#include "EntitySystem/MovieSceneEntityIDs.h"
+#include "EntitySystem/MovieSceneSequenceInstanceHandle.h"
+
+class UMovieSceneEntitySystemLinker;
+class UMovieSceneAnimationMixerTrack;
+class UMovieSceneSection;
+
+struct FMovieSceneBoneMatchData;
+
+namespace UE::UAF { struct FKeyframeState; }
+
+namespace UE::MovieScene::AnimMixerBoneMatching
+{
+
+struct FBoneMatchingContext
+{
+	UMovieSceneEntitySystemLinker* Linker = nullptr;
+	FInstanceHandle InstanceHandle;
+	UMovieSceneAnimationMixerTrack* MixerTrack = nullptr;
+
+	// Sequencer playhead time, used when MatchTimeMode == AtCurrentTime
+	FFrameTime CurrentTime;
+};
+
+// Compute the bone match transform for a target section.
+// Evaluates the ECS at the resolved match time to get both the underlying
+// pose (everything except the target, capped at the target's priority) and
+// the target section's isolated pose, then computes the selective per-axis
+// match offset between them.
+//
+// Returns a copy of Settings with MatchTransform, bIsValid, and bIsDirty updated.
+MOVIESCENEANIMMIXER_API FMovieSceneBoneMatchData ComputeBoneMatch(
+	UMovieSceneSection* TargetSection,
+	const FMovieSceneBoneMatchData& Settings,
+	const FBoneMatchingContext& Context);
+
+// Extract a bone's component-space transform from a baked keyframe.
+// Walks up the parent chain from the target bone to root, accumulating
+// local transforms into component space. Returns identity if the bone
+// is not found or the keyframe is null.
+MOVIESCENEANIMMIXER_API FTransform ExtractBoneComponentSpaceTransform(
+	const UE::UAF::FKeyframeState* Keyframe,
+	FName BoneName);
+
+// After a section's bone match is updated, propagate rematches to any
+// sections whose ReferenceSection points to it. Recurses through the
+// dependency chain using a visited set to prevent infinite loops.
+MOVIESCENEANIMMIXER_API void PropagateRematch(
+	UMovieSceneSection* UpdatedSection,
+	const FBoneMatchingContext& Context,
+	TSet<UMovieSceneSection*>& Visited);
+
+} // namespace UE::MovieScene::AnimMixerBoneMatching
diff --git a/Engine/Plugins/Experimental/MovieSceneAnimMixer/Source/MovieSceneAnimMixer/Public/AnimMixerComponentTypes.h b/Engine/Plugins/Experimental/MovieSceneAnimMixer/Source/MovieSceneAnimMixer/Public/AnimMixerComponentTypes.h
--- a/Engine/Plugins/Experimental/MovieSceneAnimMixer/Source/MovieSceneAnimMixer/Public/AnimMixerComponentTypes.h
+++ b/Engine/Plugins/Experimental/MovieSceneAnimMixer/Source/MovieSceneAnimMixer/Public/AnimMixerComponentTypes.h
@@ -17,6 +17,7 @@
 enum class EMovieSceneRootMotionDestination : uint8;
 
 struct FMovieSceneAnimMixerEntry;
+
 class UMovieSceneAnimationMixerLayer;
 
 namespace UE::MovieScene
@@ -45,6 +46,9 @@ namespace UE::MovieScene
 		TComponentTypeID<TObjectPtr<UMovieSceneAnimationMixerLayer>> MixerLayer;
 		TComponentTypeID<TObjectPtr<UMirrorDataTable>> MirrorTable;
 
+		TComponentTypeID<FTransform> BoneMatchTransform;
+		TComponentTypeID<EMovieSceneRootMotionGapBehavior> GapBehavior;
+
 		// The object (section or track) that owns/created this entity
 		TComponentTypeID<FObjectKey> EntityOwner;
 
diff --git a/Engine/Plugins/Experimental/MovieSceneAnimMixer/Source/MovieSceneAnimMixer/Public/MovieSceneAnimationMixerTrack.h b/Engine/Plugins/Experimental/MovieSceneAnimMixer/Source/MovieSceneAnimMixer/Public/MovieSceneAnimationMixerTrack.h
--- a/Engine/Plugins/Experimental/MovieSceneAnimMixer/Source/MovieSceneAnimMixer/Public/MovieSceneAnimationMixerTrack.h
+++ b/Engine/Plugins/Experimental/MovieSceneAnimMixer/Source/MovieSceneAnimMixer/Public/MovieSceneAnimationMixerTrack.h
@@ -8,6 +8,7 @@
 #include "Decorations/IMovieSceneChannelDecoration.h"
 #include "Decorations/IMovieSceneSectionDecoration.h"
 #include "EntitySystem/IMovieSceneEntityDecorator.h"
+#include "EntitySystem/MovieSceneEntityIDs.h"
 #include "MovieSceneMixedAnimationTarget.h"
 #include "MovieSceneNameableTrack.h"
 #include "MovieSceneAnimationMixerItemInterface.h"
@@ -15,12 +16,20 @@
 #include "Tracks/MovieSceneCommonAnimationTrack.h"
 #include "Systems/MovieSceneRootMotionSystem.h"
 #include "MovieSceneTrack.h"
+#include "MovieSceneBoneMatchChannel.h"
 #include "MovieSceneAnimationMixerTrack.generated.h"
 
 class UMovieSceneAnimationMixerLayer;
 class UMovieSceneAnimTransitionSectionBase;
 class UMovieSceneAnimCrossfadeTransitionSection;
 
+UENUM()
+enum class EMovieSceneRootMotionKeepState : uint8
+{
+	DontKeep,
+	KeepState
+};
+
 USTRUCT()
 struct FMovieSceneByteChannelDefaultOnly : public FMovieSceneByteChannel
 {
@@ -115,51 +124,48 @@ public:
 	UPROPERTY()
 	FMovieSceneByteChannelDefaultOnly TransformMode;
 
-	// Bone Matching Configuration
-
-	/** Name of the bone to match transforms to */
-	UPROPERTY(EditAnywhere, Category="Bone Matching")
-	FName MatchBoneName;
-
-	/** Whether to match with the previous section's end transform */
-	UPROPERTY(EditAnywhere, Category="Bone Matching")
-	bool bMatchWithPrevious = true;
-
-	/** Whether to match translation */
-	UPROPERTY(EditAnywhere, Category="Bone Matching")
-	bool bMatchTranslation = true;
+	// Bone Matching
 
-	/** Whether to include Z height when matching translation */
-	UPROPERTY(EditAnywhere, Category="Bone Matching", meta=(EditCondition="bMatchTranslation"))
-	bool bMatchIncludeZHeight = false;
+	UPROPERTY()
+	FMovieSceneBoneMatchChannel BoneMatchChannel;
 
-	/** Whether to match rotation yaw */
-	UPROPERTY(EditAnywhere, Category="Bone Matching")
-	bool bMatchRotationYaw = true;
+	// Read the first key from the channel, or return a default (invalid) match
+	FMovieSceneBoneMatchData GetBoneMatchData() const
+	{
+		TMovieSceneChannelData<const FMovieSceneBoneMatchData> Data = BoneMatchChannel.GetData();
+		TArrayView<const FMovieSceneBoneMatchData> Values = Data.GetValues();
+		return Values.Num() > 0 ? Values[0] : FMovieSceneBoneMatchData();
+	}
 
-	/** Whether to match rotation pitch */
-	UPROPERTY(EditAnywhere, Category="Bone Matching")
-	bool bMatchRotationPitch = false;
+	bool HasBoneMatch() const
+	{
+		return BoneMatchChannel.GetNumKeys() > 0;
+	}
 
-	/** Whether to match rotation roll */
-	UPROPERTY(EditAnywhere, Category="Bone Matching")
-	bool bMatchRotationRoll = false;
+	// Keep State
 
-	/** Computed location offset from bone matching (read-only) */
-	UPROPERTY(VisibleAnywhere, Category="Bone Matching")
-	FVector MatchedLocationOffset = FVector::ZeroVector;
+	UPROPERTY()
+	FMovieSceneByteChannelDefaultOnly KeepStateChannel;
 
-	/** Computed rotation offset from bone matching (read-only) */
-	UPROPERTY(VisibleAnywhere, Category="Bone Matching")
-	FRotator MatchedRotationOffset = FRotator::ZeroRotator;
+	bool GetKeepStateAfterSectionEnds() const
+	{
+		return KeepStateChannel.GetDefault().Get((uint8)EMovieSceneRootMotionKeepState::DontKeep)
+			== (uint8)EMovieSceneRootMotionKeepState::KeepState;
+	}
 
-	// Keep State
+	// Cached initial root transform at the section's start time, computed via
+	// isolated bake evaluation. Used by AccumulatedOffset mode to compute the
+	// delta from the animation's starting position.
+	UPROPERTY()
+	FTransform InitialRootTransform = FTransform::Identity;
 
-	/** If true, the transform offset persists after this section ends (for use with gap behavior) */
-	UPROPERTY(EditAnywhere, Category="Root Motion")
-	bool bKeepStateAfterSectionEnds = false;
+	UPROPERTY()
+	bool bInitialRootTransformValid = false;
 };
 
+// Returns true if setting TargetSection to reference CandidateSection
+// would create a bone match dependency cycle
+MOVIESCENEANIMMIXER_API bool WouldCreateBoneMatchCycle(UMovieSceneSection* CandidateSection, UMovieSceneSection* TargetSection);
 
 /**
  *
@@ -191,10 +197,24 @@ public:
 	MOVIESCENEANIMMIXER_API virtual void AddSectionWithAutoRow(UMovieSceneSection* Section) override;
 	MOVIESCENEANIMMIXER_API virtual int32 ComputeNextAvailableRowIndex() const override;
 
+	// Invalidate the accumulated offset cache on the root motion section.
+	// Call whenever sections change (add/remove/move) to mark the cache for rebuild.
+	MOVIESCENEANIMMIXER_API void InvalidateAccumulatedOffsetCache();
+
+	MOVIESCENEANIMMIXER_API bool HasDirtyAccumulatedOffsetCache() const;
+
+	// Rebuild the accumulated offset cache on the root motion section if dirty.
+	// Must be called from outside the ECS evaluation pipeline — the rebuild uses
+	// isolated bake evaluations that re-enter the ECS.
+	MOVIESCENEANIMMIXER_API void RebuildDirtyAccumulatedOffsetCache(
+		class UMovieSceneEntitySystemLinker* Linker,
+		UE::MovieScene::FInstanceHandle InstanceHandle);
+
 #if WITH_EDITORONLY_DATA
 	virtual FText GetTrackRowDisplayName(int32 RowIndex) const override;
 	virtual FText GetDefaultDisplayName() const override;
 	virtual bool CanRename() const override;
+
 #endif
 
 	MOVIESCENEANIMMIXER_API virtual UMovieSceneTrack* AddChildTrack(FGuid ObjectBinding, TSubclassOf<UMovieSceneTrack> TrackClass, int32 RowIndex) override;
diff --git a/Engine/Plugins/Experimental/MovieSceneAnimMixer/Source/MovieSceneAnimMixer/Public/MovieSceneBoneMatchChannel.h b/Engine/Plugins/Experimental/MovieSceneAnimMixer/Source/MovieSceneAnimMixer/Public/MovieSceneBoneMatchChannel.h
new file mode 100644
--- /dev/null
+++ b/Engine/Plugins/Experimental/MovieSceneAnimMixer/Source/MovieSceneAnimMixer/Public/MovieSceneBoneMatchChannel.h
@@ -0,0 +1,65 @@
+// Copyright Epic Games, Inc. All Rights Reserved.
+
+#pragma once
+
+#include "Misc/FrameNumber.h"
+#include "Channels/MovieSceneChannel.h"
+#include "Channels/MovieSceneChannelData.h"
+#include "Channels/MovieSceneChannelTraits.h"
+#include "MovieSceneBoneMatchData.h"
+#include "MovieSceneBoneMatchChannel.generated.h"
+
+// Channel that stores FMovieSceneBoneMatchData as discrete (non-interpolated) keys.
+// Typically holds 0 or 1 keys: one bone match per section.
+USTRUCT()
+struct FMovieSceneBoneMatchChannel : public FMovieSceneChannel
+{
+	GENERATED_BODY()
+
+	typedef FMovieSceneBoneMatchData CurveValueType;
+
+	TMovieSceneChannelData<FMovieSceneBoneMatchData> GetData()
+	{
+		return TMovieSceneChannelData<FMovieSceneBoneMatchData>(&Times, &Values, this, &KeyHandles);
+	}
+
+	TMovieSceneChannelData<const FMovieSceneBoneMatchData> GetData() const
+	{
+		return TMovieSceneChannelData<const FMovieSceneBoneMatchData>(&Times, &Values);
+	}
+
+	// FMovieSceneChannel interface
+	MOVIESCENEANIMMIXER_API virtual void GetKeys(const TRange<FFrameNumber>& WithinRange, TArray<FFrameNumber>* OutKeyTimes, TArray<FKeyHandle>* OutKeyHandles) override;
+	MOVIESCENEANIMMIXER_API virtual void GetKeyTimes(TArrayView<const FKeyHandle> InHandles, TArrayView<FFrameNumber> OutKeyTimes) override;
+	MOVIESCENEANIMMIXER_API virtual void SetKeyTimes(TArrayView<const FKeyHandle> InHandles, TArrayView<const FFrameNumber> InKeyTimes) override;
+	MOVIESCENEANIMMIXER_API virtual void DuplicateKeys(TArrayView<const FKeyHandle> InHandles, TArrayView<FKeyHandle> OutNewHandles) override;
+	MOVIESCENEANIMMIXER_API virtual void DeleteKeys(TArrayView<const FKeyHandle> InHandles) override;
+	MOVIESCENEANIMMIXER_API virtual void DeleteKeysFrom(FFrameNumber InTime, bool bDeleteKeysBefore) override;
+	MOVIESCENEANIMMIXER_API virtual void RemapTimes(const UE::MovieScene::IRetimingInterface& Retimer) override;
+	MOVIESCENEANIMMIXER_API virtual TRange<FFrameNumber> ComputeEffectiveRange() const override;
+	MOVIESCENEANIMMIXER_API virtual int32 GetNumKeys() const override;
+	MOVIESCENEANIMMIXER_API virtual void Reset() override;
+	MOVIESCENEANIMMIXER_API virtual void Offset(FFrameNumber DeltaPosition) override;
+	MOVIESCENEANIMMIXER_API void NotifyReferenceChanged(FFrameNumber NewResolvedTime);
+	MOVIESCENEANIMMIXER_API virtual void Optimize(const FKeyDataOptimizationParams& InParameters) override;
+	MOVIESCENEANIMMIXER_API virtual void ClearDefault() override;
+	MOVIESCENEANIMMIXER_API virtual FKeyHandle GetHandle(int32 Index) override;
+	MOVIESCENEANIMMIXER_API virtual int32 GetIndex(FKeyHandle Handle) override;
+
+private:
+
+	UPROPERTY(meta=(KeyTimes))
+	TArray<FFrameNumber> Times;
+
+	UPROPERTY(meta=(KeyValues))
+	TArray<FMovieSceneBoneMatchData> Values;
+
+	UPROPERTY(Transient)
+	FMovieSceneKeyHandleMap KeyHandles;
+};
+
+template<>
+struct TMovieSceneChannelTraits<FMovieSceneBoneMatchChannel> : TMovieSceneChannelTraitsBase<FMovieSceneBoneMatchChannel>
+{
+	enum { SupportsDefaults = false };
+};
diff --git a/Engine/Plugins/Experimental/MovieSceneAnimMixer/Source/MovieSceneAnimMixer/Public/MovieSceneBoneMatchData.h b/Engine/Plugins/Experimental/MovieSceneAnimMixer/Source/MovieSceneAnimMixer/Public/MovieSceneBoneMatchData.h
new file mode 100644
--- /dev/null
+++ b/Engine/Plugins/Experimental/MovieSceneAnimMixer/Source/MovieSceneAnimMixer/Public/MovieSceneBoneMatchData.h
@@ -0,0 +1,95 @@
+// Copyright Epic Games, Inc. All Rights Reserved.
+
+#pragma once
+
+#include "CoreTypes.h"
+#include "Math/Transform.h"
+#include "Misc/FrameNumber.h"
+#include "UObject/WeakObjectPtr.h"
+#include "MovieSceneBoneMatchData.generated.h"
+
+class UMovieSceneSection;
+
+UENUM()
+enum class EBoneMatchTimeMode : uint8
+{
+	AtCurrentTime,
+	AtStartOfSelectedSection,
+	AtEndOfSelectedSection,
+	AtStartOfReferenceSection,
+	AtEndOfReferenceSection,
+	InBetween
+};
+
+USTRUCT()
+struct FMovieSceneBoneMatchData
+{
+	GENERATED_BODY()
+
+	// Computed match transform offset
+	UPROPERTY(VisibleAnywhere, Category="Bone Matching")
+	FTransform MatchTransform = FTransform::Identity;
+
+	// Bone to match
+	UPROPERTY(EditAnywhere, Category="Bone Matching")
+	FName BoneName;
+
+	// Reference section for determining match time
+	UPROPERTY()
+	TWeakObjectPtr<UMovieSceneSection> ReferenceSection;
+
+	UPROPERTY(EditAnywhere, Category="Bone Matching")
+	EBoneMatchTimeMode MatchTimeMode = EBoneMatchTimeMode::AtCurrentTime;
+
+	UPROPERTY(EditAnywhere, Category="Bone Matching")
+	bool bMatchLocationX = true;
+
+	UPROPERTY(EditAnywhere, Category="Bone Matching")
+	bool bMatchLocationY = true;
+
+	UPROPERTY(EditAnywhere, Category="Bone Matching")
+	bool bMatchLocationZ = false;
+
+	UPROPERTY(EditAnywhere, Category="Bone Matching")
+	bool bMatchRotationX = false;
+
+	UPROPERTY(EditAnywhere, Category="Bone Matching")
+	bool bMatchRotationY = false;
+
+	UPROPERTY(EditAnywhere, Category="Bone Matching")
+	bool bMatchRotationZ = true;
+
+	UPROPERTY()
+	bool bIsValid = false;
+
+	// Set when section moves after a valid match
+	UPROPERTY()
+	bool bIsDirty = false;
+
+	// Resolved match time (frame number) filled in by ComputeBoneMatch.
+	// Used by the editor to place the key on the channel.
+	UPROPERTY()
+	FFrameNumber MatchTime;
+
+	// Compares user-specified settings and validity state. Computed results
+	// (MatchTransform, MatchTime) and transient state (bIsDirty) are excluded
+	// since they change on every re-evaluation.
+	friend bool operator==(const FMovieSceneBoneMatchData& A, const FMovieSceneBoneMatchData& B)
+	{
+		return A.BoneName == B.BoneName
+			&& A.ReferenceSection == B.ReferenceSection
+			&& A.MatchTimeMode == B.MatchTimeMode
+			&& A.bIsValid == B.bIsValid
+			&& A.bMatchLocationX == B.bMatchLocationX
+			&& A.bMatchLocationY == B.bMatchLocationY
+			&& A.bMatchLocationZ == B.bMatchLocationZ
+			&& A.bMatchRotationX == B.bMatchRotationX
+			&& A.bMatchRotationY == B.bMatchRotationY
+			&& A.bMatchRotationZ == B.bMatchRotationZ;
+	}
+
+	friend bool operator!=(const FMovieSceneBoneMatchData& A, const FMovieSceneBoneMatchData& B)
+	{
+		return !(A == B);
+	}
+};
diff --git a/Engine/Plugins/Experimental/MovieSceneAnimMixer/Source/MovieSceneAnimMixer/Public/MovieSceneRootMotionSection.h b/Engine/Plugins/Experimental/MovieSceneAnimMixer/Source/MovieSceneAnimMixer/Public/MovieSceneRootMotionSection.h
--- a/Engine/Plugins/Experimental/MovieSceneAnimMixer/Source/MovieSceneAnimMixer/Public/MovieSceneRootMotionSection.h
+++ b/Engine/Plugins/Experimental/MovieSceneAnimMixer/Source/MovieSceneAnimMixer/Public/MovieSceneRootMotionSection.h
@@ -5,6 +5,7 @@
 #include "CoreTypes.h"
 #include "Channels/MovieSceneByteChannel.h"
 #include "EntitySystem/IMovieSceneEntityProvider.h"
+#include "EntitySystem/MovieSceneEntityIDs.h"
 #include "MovieSceneSection.h"
 #include "MovieSceneAnimationMixerTrack.h"
 #include "MovieSceneRootMotionSection.generated.h"
@@ -34,7 +35,11 @@ enum class EMovieSceneRootMotionGapBehavior : uint8
 	PersistPreviousTransform
 };
 
-/** Cached accumulated offset from sections with KeepState enabled */
+/** Cached accumulated offset from sections with KeepState enabled.
+ *  For non-AccumulatedOffset sections, AccumulatedOffset holds world-space
+ *  and AnimSpaceAccumulatedOffset holds animation-space. For AccumulatedOffset
+ *  sections, both fields hold animation-space values because the world
+ *  conversion is applied at runtime by the conversion task. */
 USTRUCT()
 struct FRootMotionAccumulatedOffset
 {
@@ -44,10 +49,14 @@ struct FRootMotionAccumulatedOffset
 	UPROPERTY()
 	FFrameNumber StartFrame;
 
-	/** The accumulated transform offset at this point */
+	/** The accumulated transform offset at this point (world space) */
 	UPROPERTY()
 	FTransform AccumulatedOffset = FTransform::Identity;
 
+	/** The accumulated transform offset in animation space (pre-world conversion) */
+	UPROPERTY()
+	FTransform AnimSpaceAccumulatedOffset = FTransform::Identity;
+
 	/** Weight of this offset for blending */
 	UPROPERTY()
 	float Weight = 1.0f;
@@ -81,11 +90,20 @@ public:
 	virtual void ImportEntityImpl(UMovieSceneEntitySystemLinker* EntityLinker, const FEntityImportParams& Params, FImportedEntity* OutImportedEntity) override;
 	virtual bool PopulateEvaluationFieldImpl(const TRange<FFrameNumber>& EffectiveRange, const FMovieSceneEvaluationFieldEntityMetaData& InMetaData, FMovieSceneEntityComponentFieldBuilder* OutFieldBuilder) override;
 
-	/** Get the accumulated offset at a given time*/
+	/** Get the accumulated offset at a given time (world space) */
 	MOVIESCENEANIMMIXER_API FTransform GetAccumulatedOffsetAt(FFrameTime Time) const;
 
-	/** Rebuild the accumulated offset cache when sections change. */
-	MOVIESCENEANIMMIXER_API void RebuildAccumulatedOffsetCache() const;
+	/** Get the accumulated offset at a given time in animation space (pre-world conversion) */
+	MOVIESCENEANIMMIXER_API FTransform GetAnimSpaceAccumulatedOffsetAt(FFrameTime Time) const;
+
+	/** Rebuild the accumulated offset cache using isolated bake evaluations.*/
+	MOVIESCENEANIMMIXER_API void RebuildAccumulatedOffsetCache(
+		class UMovieSceneEntitySystemLinker* Linker,
+		UE::MovieScene::FInstanceHandle InstanceHandle,
+		class UMovieSceneAnimationMixerTrack* MixerTrack);
+
+	void InvalidateAccumulatedOffsetCache() { bAccumulatedOffsetCacheDirty = true; }
+	bool IsAccumulatedOffsetCacheDirty() const { return bAccumulatedOffsetCacheDirty; }
 
 public:
 
@@ -93,9 +111,12 @@ public:
 	UPROPERTY(EditAnywhere, Category="Root Motion")
 	EMovieSceneRootMotionGapBehavior GapBehavior = EMovieSceneRootMotionGapBehavior::ResetToOrigin;
 
-	/** Cached accumulated offsets from sections with KeepState enabled. Mutable for lazy rebuilding. */
-	UPROPERTY(Transient)
-	mutable TArray<FRootMotionAccumulatedOffset> AccumulatedOffsetCache;
+	/** Cached accumulated offsets from sections with KeepState enabled. Serialized so
+	 *  the correct values are available at runtime without needing to rebuild. */
+	UPROPERTY()
+	TArray<FRootMotionAccumulatedOffset> AccumulatedOffsetCache;
+
+	bool bAccumulatedOffsetCacheDirty = false;
 
 private:
 
diff --git a/Engine/Plugins/Experimental/MovieSceneAnimMixer/Source/MovieSceneAnimMixer/Public/Systems/MovieSceneAnimMixerSystem.h b/Engine/Plugins/Experimental/MovieSceneAnimMixer/Source/MovieSceneAnimMixer/Public/Systems/MovieSceneAnimMixerSystem.h
--- a/Engine/Plugins/Experimental/MovieSceneAnimMixer/Source/MovieSceneAnimMixer/Public/Systems/MovieSceneAnimMixerSystem.h
+++ b/Engine/Plugins/Experimental/MovieSceneAnimMixer/Source/MovieSceneAnimMixer/Public/Systems/MovieSceneAnimMixerSystem.h
@@ -13,6 +13,7 @@
 #include "EvaluationVM/Tasks/BlendKeyframes.h"
 #include "MovieSceneAnimMixerSystem.generated.h"
 
+class UMovieSceneAnimationMixerTrack;
 struct FMovieSceneSkeletalAnimationComponentData;
 struct FAnimNextEvaluationTask;
 
@@ -46,6 +47,8 @@ struct FMovieSceneAnimMixerEntry
 	/** True if this entry is from a transition section */
 	bool bIsTransition = false;
 
+	/** Bone match transform for this entry (set only when a valid match exists) */
+	TOptional<FTransform> BoneMatchTransform;
 
 	inline bool operator<(const FMovieSceneAnimMixerEntry& RHS) const
 	{
@@ -91,6 +94,7 @@ struct FMovieSceneAnimMixer
 	TSharedPtr<UE::UAF::FEvaluationProgram> EvaluationProgram;
 	TWeakPtr<FMovieSceneMixerRootMotionComponentData> WeakRootMotion;
 
+	TWeakObjectPtr<UMovieSceneAnimationMixerTrack> OwnerTrack;
 	bool bNeedsResort = false;
 
 	/**
@@ -191,11 +195,16 @@ private:
 	TMultiMap<FObjectKey, TWeakPtr<FMovieSceneMixerRootMotionComponentData>> ActorToRootMotion;
 
 	virtual void OnLink() override;
+	virtual void OnUnlink() override;
 	virtual bool IsRelevantImpl(UMovieSceneEntitySystemLinker* InLinker) const override;
 	virtual void OnRun(FSystemTaskPrerequisites& InPrerequisites, FSystemSubsequentTasks& Subsequents) override;
 	virtual void OnSchedulePersistentTasks(UE::MovieScene::IEntitySystemScheduler* TaskScheduler) override;
 	virtual void OnCleanTaggedGarbage() override;
 
+	void HandlePreFlush(UMovieSceneEntitySystemLinker* InLinker);
+
 	UPROPERTY()
 	TObjectPtr<UMovieSceneRootMotionSystem> RootMotionSystem;
+
+	FDelegateHandle PreFlushHandle;
 };
\ No newline at end of file
diff --git a/Engine/Plugins/Experimental/MovieSceneAnimMixer/Source/MovieSceneAnimMixer/Public/Systems/MovieSceneMixerBakeTargetSystem.h b/Engine/Plugins/Experimental/MovieSceneAnimMixer/Source/MovieSceneAnimMixer/Public/Systems/MovieSceneMixerBakeTargetSystem.h
new file mode 100644
--- /dev/null
+++ b/Engine/Plugins/Experimental/MovieSceneAnimMixer/Source/MovieSceneAnimMixer/Public/Systems/MovieSceneMixerBakeTargetSystem.h
@@ -0,0 +1,59 @@
+// Copyright Epic Games, Inc. All Rights Reserved.
+
+#pragma once
+
+#include "CoreTypes.h"
+#include "AnimMixerBakeEvaluation.h"
+#include "EntitySystem/MovieSceneEntitySystem.h"
+#include "MovieSceneMixedAnimationTarget.h"
+#include "MovieSceneMixerBakeTargetSystem.generated.h"
+
+struct FAnimNextEvaluationTask;
+
+namespace UE::UAF
+{
+	struct FEvaluationVM;
+	struct FKeyframeState;
+}
+
+// Hidden target struct used during bake evaluation to tell the target system
+// to cache results instead of applying them. Never user-facing.
+USTRUCT(meta=(Hidden))
+struct FMovieSceneMixerBakeTarget : public FMovieSceneMixedAnimationTarget
+{
+	GENERATED_BODY()
+
+	UE::MovieScene::AnimMixerBakeEvaluation::FBakeFilter Filter;
+};
+
+// Entity system that captures mixer evaluation results during bake evaluation.
+// When bBakeActive is true, this system executes mixer programs in a temporary
+// VM and stores results (root motion transform + keyframe) instead of applying
+// them to any skeletal mesh.
+UCLASS(MinimalAPI)
+class UMovieSceneMixerBakeTargetSystem : public UMovieSceneEntitySystem
+{
+	GENERATED_BODY()
+
+public:
+
+	UMovieSceneMixerBakeTargetSystem(const FObjectInitializer& Init);
+
+	// Result from a single bake evaluation frame. One per mixer entity that
+	// matches FMovieSceneMixerBakeTarget during a flush. 
+	struct FBakedResult
+	{
+		TUniquePtr<UE::UAF::FKeyframeState> Keyframe;
+		FTransform RootMotionTransform = FTransform::Identity;
+	};
+
+	// Accumulated baked results, populated during bake evaluation
+	TArray<FBakedResult> BakedResults;
+
+	// When true, the system captures results instead of being a no-op
+	bool bBakeActive = false;
+
+private:
+
+	virtual void OnSchedulePersistentTasks(UE::MovieScene::IEntitySystemScheduler* TaskScheduler) override;
+};
diff --git a/Engine/Plugins/Experimental/MovieSceneAnimMixer/Source/MovieSceneAnimMixer/Public/Systems/MovieSceneRootMotionSystem.h b/Engine/Plugins/Experimental/MovieSceneAnimMixer/Source/MovieSceneAnimMixer/Public/Systems/MovieSceneRootMotionSystem.h
--- a/Engine/Plugins/Experimental/MovieSceneAnimMixer/Source/MovieSceneAnimMixer/Public/Systems/MovieSceneRootMotionSystem.h
+++ b/Engine/Plugins/Experimental/MovieSceneAnimMixer/Source/MovieSceneAnimMixer/Public/Systems/MovieSceneRootMotionSystem.h
@@ -10,8 +10,10 @@
 #include "MovieSceneRootMotionSystem.generated.h"
 
 enum class EMovieSceneRootMotionDestination : uint8;
+enum class EMovieSceneRootMotionGapBehavior : uint8;
 
 class UMovieSceneAnimMixerSystem;
+class UMovieSceneRootMotionSection;
 class UMovieSceneRootMotionSystem;
 
 // Enum describing what space that root motion should be applied in.
@@ -29,9 +31,19 @@ enum class EMovieSceneRootMotionSpace : uint8
 UENUM()
 enum class EMovieSceneRootMotionTransformMode : uint8
 {
+	// Uses root motion exclusively from the asset as is
 	Asset,
+	// Adds on the offset field onto the root motion from the asset
 	Offset,
+	// Uses the offset field as the entire root motion, ignoring root motion from the asset
 	Override,
+	// Uses the accumulated offset from previous KeepState sections as a starting position.
+	// The section's own root motion plays on top of this accumulated transform.
+	AccumulatedOffset,
+	// Uses a bone-match-derived offset to position the section so a specific bone
+	// aligns with the underlying pose at the match time. Set automatically by the
+	// bone match dialog. Manual offsets are additive on top.
+	BoneMatch,
 };
 
 USTRUCT()
@@ -108,6 +120,11 @@ public:
 
 	bool bComponentSpaceRoot = false;
 	bool bActorTransformSet = false;
+
+	// Gap behavior fields
+	EMovieSceneRootMotionGapBehavior GapBehavior{};
+	FTransform CachedGapTransform = FTransform::Identity;
+	TWeakObjectPtr<UMovieSceneRootMotionSection> WeakRootMotionSection;
 };
 
 // Task that converts the root motion attribute on the top pose of the pose stack to world space by adding on the actor transformation, root base transform, and/or transform origin.
@@ -134,6 +151,10 @@ struct FAnimNextConvertRootMotionToWorldSpaceTask : public FAnimNextEvaluationTa
 		RootTransformOffset = 1 << 4,
 		/** Completely override the root transform with RootTransform */
 		RootTransformOverride = 1 << 5,
+		/** Pre-apply accumulated offset as starting position */
+		AccumulatedOffset = 1 << 6,
+		/** Apply bone-match-derived offset to align a specific bone with the underlying pose */
+		BoneMatchOffset = 1 << 7,
 	};
 
 	FAnimNextConvertRootMotionToWorldSpaceTask() = default;
@@ -151,8 +172,18 @@ struct FAnimNextConvertRootMotionToWorldSpaceTask : public FAnimNextEvaluationTa
 	FTransform TransformOrigin;
 	/** Origin around which to apply RootTransform when space conversion is RootTransformOffset. */
 	FVector RootOffsetOrigin;
+	/** Pre-transform for AccumulatedOffset mode - shifts animation origin to accumulated position */
+	FTransform AccumulatedTransform;
+	/** Initial root transform of this section evaluated in isolation at its start time.
+	 *  Used with AccumulatedOffset to compute the delta: Delta = Inv(InitialAnimRootTransform) * N_t */
+	FTransform InitialAnimRootTransform;
+	/** Bone-match-derived root offset. Applied as post-multiply in animation space
+	 *  so the matched bone aligns with the underlying pose at the match time. */
+	FTransform BoneMatchTransform;
 	/** Bitmask specifying the conversions to perform */
 	ESpaceConversions Conversions;
+	/** When set (bake only), the pre-conversion animation-space root is written here */
+	TSharedPtr<FTransform> CaptureAnimSpaceRoot;
 };
 ENUM_CLASS_FLAGS(FAnimNextConvertRootMotionToWorldSpaceTask::ESpaceConversions);
 
@@ -179,6 +210,29 @@ struct FAnimNextStoreRootTransformTask : public FAnimNextEvaluationTask
 };
 
 
+// Task that sets the RootTransformAttribute on the top keyframe, creating it
+// if it doesn't exist. Used for gap behavior (PersistPreviousTransform) where
+// we need to inject a cached world-space transform into an empty or
+// non-root-motion pose. Unlike FAnimNextConvertRootMotionToWorldSpaceTask,
+// this doesn't require an existing attribute and performs no space conversion
+// (the gap transform is already in the correct post-conversion space from the
+// trajectory cache).
+USTRUCT()
+struct FAnimNextSetRootTransformAttributeTask : public FAnimNextEvaluationTask
+{
+	GENERATED_BODY()
+
+	DECLARE_ANIM_EVALUATION_TASK(FAnimNextSetRootTransformAttributeTask)
+
+	FAnimNextSetRootTransformAttributeTask() = default;
+
+	static FAnimNextSetRootTransformAttributeTask Make(const FTransform& InTransform);
+
+	virtual void Execute(UE::UAF::FEvaluationVM& VM) const override;
+
+	FTransform Transform;
+};
+
 // Takes in evaluation tasks on mixers.
 // Mixes just the root motion attributes.
 // Converts it from animation space to either additive actor or component space (based on which attribute used).
diff --git a/Engine/Plugins/Experimental/MovieSceneAnimMixer/Source/MovieSceneAnimMixerEditor/MovieSceneAnimMixerEditor.Build.cs b/Engine/Plugins/Experimental/MovieSceneAnimMixer/Source/MovieSceneAnimMixerEditor/MovieSceneAnimMixerEditor.Build.cs
--- a/Engine/Plugins/Experimental/MovieSceneAnimMixer/Source/MovieSceneAnimMixerEditor/MovieSceneAnimMixerEditor.Build.cs
+++ b/Engine/Plugins/Experimental/MovieSceneAnimMixer/Source/MovieSceneAnimMixerEditor/MovieSceneAnimMixerEditor.Build.cs
@@ -23,6 +23,8 @@ namespace UnrealBuildTool.Rules
 					"AssetRegistry",
 					"BlueprintGraph",
 					"Engine",
+					"InputCore",
+					"Persona",
 					"MovieScene",
 					"MovieSceneAnimMixer",
 					"MovieSceneTracks",
diff --git a/Engine/Plugins/Experimental/MovieSceneAnimMixer/Source/MovieSceneAnimMixerEditor/Private/MovieSceneAnimMixerEditorModule.cpp b/Engine/Plugins/Experimental/MovieSceneAnimMixer/Source/MovieSceneAnimMixerEditor/Private/MovieSceneAnimMixerEditorModule.cpp
--- a/Engine/Plugins/Experimental/MovieSceneAnimMixer/Source/MovieSceneAnimMixerEditor/Private/MovieSceneAnimMixerEditorModule.cpp
+++ b/Engine/Plugins/Experimental/MovieSceneAnimMixer/Source/MovieSceneAnimMixerEditor/Private/MovieSceneAnimMixerEditorModule.cpp
@@ -9,6 +9,7 @@
 #include "IMovieSceneAnimMixerTargetMenuProvider.h"
 #include "ISequencerDecorationEditor.h"
 #include "ISequencerModule.h"
+#include "MVVM/Views/KeyDrawParams.h"
 #include "Modules/ModuleInterface.h"
 #include "Modules/ModuleManager.h"
 #include "MovieSceneAnimationMixerTrack.h"
@@ -16,6 +17,7 @@
 #include "MovieSceneAnimCrossfadeTransitionSection.h"
 #include "MovieSceneAnimInertialDeadBlendTransitionSection.h"
 #include "RootMotionTargetDecorationEditor.h"
+#include "RootMotionSettingsDecorationEditor.h"
 #include "Channels/BuiltInChannelEditors.h"
 #include "MVVM/AnimationMixerTrackModel.h"
 
@@ -80,8 +82,75 @@ namespace Sequencer
 	{
 		return true;
 	}
+	// FMovieSceneBoneMatchChannel overloads — discrete, non-keyable channel for bone match data.
+	// These no-op implementations prevent the default template from trying to call Evaluate().
+
+	FKeyHandle EvaluateAndAddKey(FMovieSceneBoneMatchChannel* InChannel, const TMovieSceneChannelData<FMovieSceneBoneMatchData>& InChannelData, FFrameNumber InTime, ISequencer& InSequencer, FMovieSceneBoneMatchData InDefaultValue = FMovieSceneBoneMatchData())
+	{
+		return FKeyHandle::Invalid();
+	}
+
+	FKeyHandle AddOrUpdateKey(
+		FMovieSceneBoneMatchChannel* InChannel,
+		UMovieSceneSection*          InSectionToKey,
+		FFrameNumber                 InTime,
+		ISequencer&                  InSequencer,
+		const FGuid&                 InObjectBindingID,
+		FTrackInstancePropertyBindings* InPropertyBindings)
+	{
+		return FKeyHandle::Invalid();
+	}
+
+	void CopyKeys(FMovieSceneBoneMatchChannel* InChannel, const UMovieSceneSection* InSection, FName KeyAreaName, FMovieSceneClipboardBuilder& ClipboardBuilder, TArrayView<const FKeyHandle> InHandles)
+	{
+	}
+
+	void PasteKeys(FMovieSceneBoneMatchChannel* InChannel, UMovieSceneSection* Section, const FMovieSceneClipboardKeyTrack& KeyTrack, const FMovieSceneClipboardEnvironment& SrcEnvironment, const FSequencerPasteEnvironment& DstEnvironment, TArray<FKeyHandle>& OutPastedKeys)
+	{
+	}
+
+	bool SupportsCurveEditorModels(const TMovieSceneChannelHandle<FMovieSceneBoneMatchChannel>& ChannelHandle)
+	{
+		return false;
+	}
+
+	bool CanCreateKeyEditor(const FMovieSceneBoneMatchChannel* Channel)
+	{
+		return false;
+	}
+
+	void DrawKeys(FMovieSceneBoneMatchChannel* Channel, TArrayView<const FKeyHandle> InKeyHandles, const UMovieSceneSection* InOwner, TArrayView<FKeyDrawParams> OutKeyDrawParams)
+	{
+		FKeyDrawParams ValidParams, InvalidParams;
+		ValidParams.FillBrush = ValidParams.BorderBrush = FAppStyle::Get().GetBrush("Sequencer.KeyDiamond");
+		ValidParams.FillTint = FLinearColor(0.2f, 0.8f, 0.2f, 1.0f); // green
+
+		InvalidParams.FillBrush = InvalidParams.BorderBrush = FAppStyle::Get().GetBrush("Sequencer.KeyDiamond");
+		InvalidParams.FillTint = FLinearColor(1.0f, 0.2f, 0.2f, 1.0f); // red
+
+		TMovieSceneChannelData<FMovieSceneBoneMatchData> ChannelData = Channel->GetData();
+		TArrayView<FMovieSceneBoneMatchData> Values = ChannelData.GetValues();
+
+		for (int32 Index = 0; Index < InKeyHandles.Num(); ++Index)
+		{
+			int32 KeyIndex = ChannelData.GetIndex(InKeyHandles[Index]);
+			if (KeyIndex != INDEX_NONE && Values[KeyIndex].bIsValid)
+			{
+				OutKeyDrawParams[Index] = ValidParams;
+			}
+			else
+			{
+				OutKeyDrawParams[Index] = InvalidParams;
+			}
+		}
+	}
 } // namespace Sequencer
 
+inline bool EvaluateChannel(const FMovieSceneBoneMatchChannel* InChannel, FFrameTime InTime, FMovieSceneBoneMatchData& OutValue)
+{
+	return false;
+}
+
 
 // Include order matters here since this header needs to be included after the above overloads in order for ADL to see them correctly
 #include "SequencerChannelInterface.h"
@@ -111,6 +180,7 @@ namespace UE::MovieScene
 		));
 
 		SequencerModule.RegisterChannelInterface<FMovieSceneByteChannelDefaultOnly>();
+		SequencerModule.RegisterChannelInterface<FMovieSceneBoneMatchChannel>();
 
 		// Register custom section interface for transition sections
 		UE::Sequencer::FAnimationMixerTrackEditor::RegisterCustomMixerAnimSection(
@@ -141,11 +211,19 @@ namespace UE::MovieScene
 		// Register decoration editors
 		RootMotionTargetDecorationEditor = MakeUnique<UE::Sequencer::FRootMotionTargetDecorationEditor>();
 		IModularFeatures::Get().RegisterModularFeature(ISequencerDecorationEditor::GetModularFeatureName(), RootMotionTargetDecorationEditor.Get());
+
+		RootMotionSettingsDecorationEditor = MakeUnique<UE::Sequencer::FRootMotionSettingsDecorationEditor>();
+		IModularFeatures::Get().RegisterModularFeature(ISequencerDecorationEditor::GetModularFeatureName(), RootMotionSettingsDecorationEditor.Get());
 	}
 
 	void FMovieSceneAnimMixerEditorModule::ShutdownModule()
 	{
 		// Unregister decoration editors
+		if (RootMotionSettingsDecorationEditor)
+		{
+			RootMotionSettingsDecorationEditor->UnsubscribeAll();
+			IModularFeatures::Get().UnregisterModularFeature(ISequencerDecorationEditor::GetModularFeatureName(), RootMotionSettingsDecorationEditor.Get());
+		}
 		if (RootMotionTargetDecorationEditor)
 		{
 			IModularFeatures::Get().UnregisterModularFeature(ISequencerDecorationEditor::GetModularFeatureName(), RootMotionTargetDecorationEditor.Get());
diff --git a/Engine/Plugins/Experimental/MovieSceneAnimMixer/Source/MovieSceneAnimMixerEditor/Private/MovieSceneAnimMixerEditorModule.h b/Engine/Plugins/Experimental/MovieSceneAnimMixer/Source/MovieSceneAnimMixerEditor/Private/MovieSceneAnimMixerEditorModule.h
--- a/Engine/Plugins/Experimental/MovieSceneAnimMixer/Source/MovieSceneAnimMixerEditor/Private/MovieSceneAnimMixerEditorModule.h
+++ b/Engine/Plugins/Experimental/MovieSceneAnimMixer/Source/MovieSceneAnimMixerEditor/Private/MovieSceneAnimMixerEditorModule.h
@@ -14,6 +14,7 @@ class FAnimNextInjectionTargetMenuProvider;
 namespace UE::Sequencer
 {
 	class FRootMotionTargetDecorationEditor;
+	class FRootMotionSettingsDecorationEditor;
 }
 
 namespace UE::MovieScene
@@ -40,6 +41,7 @@ namespace UE::MovieScene
 
 		// Decoration editors
 		TUniquePtr<UE::Sequencer::FRootMotionTargetDecorationEditor> RootMotionTargetDecorationEditor;
+		TUniquePtr<UE::Sequencer::FRootMotionSettingsDecorationEditor> RootMotionSettingsDecorationEditor;
 	};
 }
 
diff --git a/Engine/Plugins/Experimental/MovieSceneAnimMixer/Source/MovieSceneAnimMixerEditor/Private/MovieSceneAnimationMixerTrackEditor.cpp b/Engine/Plugins/Experimental/MovieSceneAnimMixer/Source/MovieSceneAnimMixerEditor/Private/MovieSceneAnimationMixerTrackEditor.cpp
--- a/Engine/Plugins/Experimental/MovieSceneAnimMixer/Source/MovieSceneAnimMixerEditor/Private/MovieSceneAnimationMixerTrackEditor.cpp
+++ b/Engine/Plugins/Experimental/MovieSceneAnimMixer/Source/MovieSceneAnimMixerEditor/Private/MovieSceneAnimationMixerTrackEditor.cpp
@@ -2,7 +2,9 @@
 
 #include "MovieSceneAnimationMixerTrackEditor.h"
 
+#include "MovieScene.h"
 #include "Animation/AnimSequenceBase.h"
+#include "Evaluation/MovieSceneEvaluationTemplateInstance.h"
 #include "UObject/UObjectIterator.h"
 #include "AssetRegistry/AssetRegistryModule.h"
 #include "ClassIconFinder.h"
@@ -198,7 +200,10 @@ public:
 
 	FAnimMixerAnimationSection(UMovieSceneSection& InSection, TWeakPtr<ISequencer> InSequencer)
 		: FCommonAnimationSection(InSection, InSequencer)
-	{}
+	{
+	}
+
+	virtual bool ShouldShowLegacyBoneMatching() const override { return false; }
 
 	virtual int32 OnPaintSection(FSequencerSectionPainter& Painter) const override
 	{
@@ -263,7 +268,7 @@ private:
 					// Add padding with a small margin
 					return TransitionPixels + 4.0f;
 				}
-				break;  // Only need to handle one transition at the start
+				break;
 			}
 		}
 
@@ -279,18 +284,8 @@ public:
 	{
 		FCommonAnimationSection::BuildSectionContextMenu(MenuBuilder, ObjectBinding);
 
-		UMovieSceneSkeletalAnimationSection* AnimSection = Cast<UMovieSceneSkeletalAnimationSection>(WeakSection.Get());
-		if (AnimSection)
-		{
-			MenuBuilder.AddSubMenu(
-				LOCTEXT("RootTransform_Label", "Root Transform"),
-				LOCTEXT("RootTransform_Tooltip", "Options for root transform behavior from this anim clip"),
-				FNewMenuDelegate::CreateSP(this, &FAnimMixerAnimationSection::PopulateRootTransformMenu)
-			);
-		}
-
 		// Add "Create Transition To..." submenu for overlapping sections
-		UMovieSceneSection* Section = WeakSection.Get();
+		UMovieSceneSection* Section = GetSectionObject();
 		if (Section)
 		{
 			UMovieSceneAnimationMixerTrack* MixerTrack = Cast<UMovieSceneAnimationMixerTrack>(Section->GetTypedOuter<UMovieSceneTrack>());
@@ -309,167 +304,9 @@ public:
 		}
 	}
 
-	UMovieSceneRootMotionSettingsDecoration* FindRootDecoration() const
-	{
-		UMovieSceneSkeletalAnimationSection* AnimSection = Cast<UMovieSceneSkeletalAnimationSection>(WeakSection.Get());
-		return AnimSection ? AnimSection->FindDecoration<UMovieSceneRootMotionSettingsDecoration>() : nullptr;
-	}
-
-	UMovieSceneRootMotionSettingsDecoration* FindOrCreateRootDecoration() const
-	{
-		UMovieSceneSkeletalAnimationSection* AnimSection = Cast<UMovieSceneSkeletalAnimationSection>(WeakSection.Get());
-		if (AnimSection)
-		{
-			UMovieSceneRootMotionSettingsDecoration* Decoration = AnimSection->FindDecoration<UMovieSceneRootMotionSettingsDecoration>();
-			if (!Decoration)
-			{
-				AnimSection->Modify();
-				Decoration = AnimSection->GetOrCreateDecoration<UMovieSceneRootMotionSettingsDecoration>();
-			}
-			return Decoration;
-		}
-		return nullptr;
-	}
-
-	void PopulateRootTransformMenu(FMenuBuilder& MenuBuilder)
-	{
-		MenuBuilder.BeginSection(NAME_None, LOCTEXT("RootTransformModeSection_Label", "Root Transform Mode"));
-		{
-			FCanExecuteAction AlwaysExecute = FCanExecuteAction::CreateLambda([] { return true; });
-
-			MenuBuilder.AddMenuEntry(
-				LOCTEXT("RootTransformMode_None", "From Animation"),
-				LOCTEXT("RootTransformMode_None_Tip", "Use the root motion transform directly from the animation asset"),
-				FSlateIcon(),
-				FUIAction(
-					FExecuteAction::CreateSP(this, &FAnimMixerAnimationSection::SetRootTransformNone),
-					AlwaysExecute,
-					FIsActionChecked::CreateSP(this, &FAnimMixerAnimationSection::IsTransformMode, EMovieSceneRootMotionTransformMode::Asset)
-				),
-				NAME_None,
-				EUserInterfaceActionType::RadioButton
-			);
-
-			MenuBuilder.AddMenuEntry(
-				LOCTEXT("RootTransformMode_Offset", "Offset From Animation"),
-				LOCTEXT("RootTransformMode_Offset_Toolitp", "Offset the root motion transform using keyframed values"),
-				FSlateIcon(),
-				FUIAction(
-					FExecuteAction::CreateSP(this, &FAnimMixerAnimationSection::SetRootTransformOffset),
-					AlwaysExecute,
-					FIsActionChecked::CreateSP(this, &FAnimMixerAnimationSection::IsTransformMode, EMovieSceneRootMotionTransformMode::Offset)
-				),
-				NAME_None,
-				EUserInterfaceActionType::RadioButton
-			);
-
-			MenuBuilder.AddMenuEntry(
-				LOCTEXT("RootTransformMode_Override", "Manual Override"),
-				LOCTEXT("RootTransformMode_Override_Toolitp", "Completely override the root motion transform using keyframed values"),
-				FSlateIcon(),
-				FUIAction(
-					FExecuteAction::CreateSP(this, &FAnimMixerAnimationSection::SetRootTransformOverride),
-					AlwaysExecute,
-					FIsActionChecked::CreateSP(this, &FAnimMixerAnimationSection::IsTransformMode, EMovieSceneRootMotionTransformMode::Override)
-				),
-				NAME_None,
-				EUserInterfaceActionType::RadioButton
-			);
-		}
-		MenuBuilder.EndSection();
-
-		UMovieSceneRootMotionSettingsDecoration* Decoration = FindRootDecoration();
-		if (Decoration && Decoration->GetRootTransformMode() == EMovieSceneRootMotionTransformMode::Offset)
-		{
-			MenuBuilder.BeginSection(NAME_None, LOCTEXT("RootTransformOffset_Label", "Offset"));
-			{
-				MenuBuilder.AddMenuEntry(
-					LOCTEXT("RecenterRootTransform", "Re-center Root Transform"),
-					LOCTEXT("RecenterRootTransform_Tooltip", "Center the root transform for this animation around its current position"),
-					FSlateIcon(),
-					FUIAction(
-						FExecuteAction::CreateSP(this, &FAnimMixerAnimationSection::RecenterRootTransform)
-					)
-				);
-			}
-		}
-	}
-
-	bool IsTransformMode(EMovieSceneRootMotionTransformMode Mode) const
-	{
-		if (UMovieSceneRootMotionSettingsDecoration* Decoration = FindRootDecoration())
-		{
-			return Decoration->GetRootTransformMode() == Mode;
-		}
-		return Mode == EMovieSceneRootMotionTransformMode::Asset;
-	}
-
-	void SetRootTransformNone()
-	{
-		FScopedTransaction Transaction(LOCTEXT("SetRootTransformNone", "Use Root Transform from Asset"));
-
-		UMovieSceneSkeletalAnimationSection*         AnimSection = Cast<UMovieSceneSkeletalAnimationSection>(WeakSection.Get());
-		UMovieSceneRootMotionSettingsDecoration* Decoration  = FindOrCreateRootDecoration();
-		if (Decoration && AnimSection)
-		{
-			Decoration->Modify();
-			Decoration->TransformMode.SetDefault((uint8)EMovieSceneRootMotionTransformMode::Asset);
-			AnimSection->InvalidateChannelProxy();
-		}
-	}
-	void SetRootTransformOffset()
-	{
-		FScopedTransaction Transaction(LOCTEXT("SetRootTransformOffset", "Offset Root Transform from Asset"));
-
-		UMovieSceneSkeletalAnimationSection*         AnimSection = Cast<UMovieSceneSkeletalAnimationSection>(WeakSection.Get());
-		UMovieSceneRootMotionSettingsDecoration* Decoration  = FindOrCreateRootDecoration();
-		if (Decoration && AnimSection)
-		{
-			Decoration->Modify();
-			Decoration->TransformMode.SetDefault((uint8)EMovieSceneRootMotionTransformMode::Offset);
-			AnimSection->InvalidateChannelProxy();
-		}
-	}
-	void SetRootTransformOverride()
-	{
-		FScopedTransaction Transaction(LOCTEXT("SetRootTransformOverride", "Override Root Transform"));
-
-		UMovieSceneSkeletalAnimationSection*         AnimSection = Cast<UMovieSceneSkeletalAnimationSection>(WeakSection.Get());
-		UMovieSceneRootMotionSettingsDecoration* Decoration  = FindOrCreateRootDecoration();
-		if (Decoration && AnimSection)
-		{
-			Decoration->Modify();
-			Decoration->TransformMode.SetDefault((uint8)EMovieSceneRootMotionTransformMode::Override);
-			AnimSection->InvalidateChannelProxy();
-		}
-	}
-	void RecenterRootTransform()
-	{
-		FScopedTransaction Transaction(LOCTEXT("RecenterRootTransform", "Re-center Root Transform"));
-
-		TSharedPtr<ISequencer> SequencerPtr = Sequencer.Pin();
-
-		UMovieSceneSkeletalAnimationSection*         AnimSection = Cast<UMovieSceneSkeletalAnimationSection>(WeakSection.Get());
-		UMovieSceneRootMotionSettingsDecoration* Decoration  = FindRootDecoration();
-		UAnimSequenceBase*                           Animation   = AnimSection ? AnimSection->GetAnimation() : nullptr;
-
-		if (SequencerPtr && Decoration && Animation)
-		{
-			Decoration->Modify();
-
-			FQualifiedFrameTime Time = SequencerPtr->GetLocalTime();
-			double AnimTime = AnimSection->MapTimeToAnimation(Time.Time, Time.Rate);
-
-			FAnimExtractContext CurrentExtractContext(AnimTime);
-			FTransform CurrentTransform  = Animation->ExtractRootTrackTransform(CurrentExtractContext, nullptr);
-
-			Decoration->RootOriginLocation = CurrentTransform.GetTranslation();
-		}
-	}
-
 	void PopulateCreateTransitionMenu(FMenuBuilder& MenuBuilder)
 	{
-		UMovieSceneSection* Section = WeakSection.Get();
+		UMovieSceneSection* Section = GetSectionObject();
 		if (!Section)
 		{
 			return;
@@ -547,7 +384,7 @@ public:
 
 	void PopulateTransitionTypeMenu(FMenuBuilder& MenuBuilder, UMovieSceneSection* OtherSection, TArray<UClass*> TransitionClasses)
 	{
-		UMovieSceneSection* Section = WeakSection.Get();
+		UMovieSceneSection* Section = GetSectionObject();
 		if (!Section || !OtherSection)
 		{
 			return;
diff --git a/Engine/Plugins/Experimental/MovieSceneAnimMixer/Source/MovieSceneAnimMixerEditor/Private/MovieSceneAnimationMixerTrackEditor.h b/Engine/Plugins/Experimental/MovieSceneAnimMixer/Source/MovieSceneAnimMixerEditor/Private/MovieSceneAnimationMixerTrackEditor.h
--- a/Engine/Plugins/Experimental/MovieSceneAnimMixer/Source/MovieSceneAnimMixerEditor/Private/MovieSceneAnimationMixerTrackEditor.h
+++ b/Engine/Plugins/Experimental/MovieSceneAnimMixer/Source/MovieSceneAnimMixerEditor/Private/MovieSceneAnimationMixerTrackEditor.h
@@ -4,6 +4,7 @@
 
 #include "Templates/SharedPointer.h"
 #include "Templates/SubclassOf.h"
+#include "Containers/Ticker.h"
 #include "MovieSceneTrackEditor.h"
 #include "MovieSceneMixedAnimationTarget.h"
 #include "StructUtils/InstancedStruct.h"
diff --git a/Engine/Plugins/Experimental/MovieSceneAnimMixer/Source/MovieSceneAnimMixerEditor/Private/RootMotionSettingsDecorationEditor.cpp b/Engine/Plugins/Experimental/MovieSceneAnimMixer/Source/MovieSceneAnimMixerEditor/Private/RootMotionSettingsDecorationEditor.cpp
new file mode 100644
--- /dev/null
+++ b/Engine/Plugins/Experimental/MovieSceneAnimMixer/Source/MovieSceneAnimMixerEditor/Private/RootMotionSettingsDecorationEditor.cpp
@@ -0,0 +1,743 @@
+// Copyright Epic Games, Inc. All Rights Reserved.
+
+#include "RootMotionSettingsDecorationEditor.h"
+
+#include "AnimMixerBakeEvaluation.h"
+#include "AnimMixerBoneMatching.h"
+#include "Evaluation/MovieSceneEvaluationTemplateInstance.h"
+#include "Framework/MultiBox/MultiBoxBuilder.h"
+#include "ISequencer.h"
+#include "MovieSceneAnimationMixerTrack.h"
+#include "MovieSceneAnimTransitionSectionBase.h"
+#include "MovieSceneSection.h"
+#include "SBoneMatchDialog.h"
+#include "ScopedTransaction.h"
+#include "SEnumCombo.h"
+#include "Channels/MovieSceneChannelHandle.h"
+#include "MovieSceneToolHelpers.h"
+#include "MVVM/Views/SKeyNavigationButtons.h"
+#include "Systems/MovieSceneRootMotionSystem.h"
+
+#define LOCTEXT_NAMESPACE "RootMotionSettingsDecorationEditor"
+
+namespace UE::Sequencer
+{
+
+// Returns true if the given section has at least one overlapping non-transition
+// section on the same mixer track that wouldn't create a bone match dependency cycle.
+static bool HasOverlappingReferenceSections(UMovieSceneSection* Section)
+{
+	UMovieSceneAnimationMixerTrack* MixerTrack = Cast<UMovieSceneAnimationMixerTrack>(Section->GetTypedOuter<UMovieSceneTrack>());
+	if (!MixerTrack)
+	{
+		return false;
+	}
+
+	for (UMovieSceneSection* Other : MixerTrack->GetAllSections())
+	{
+		if (Other && Other != Section
+			&& !Other->IsA<UMovieSceneAnimTransitionSectionBase>()
+			&& !TRange<FFrameNumber>::Intersection(Section->GetRange(), Other->GetRange()).IsEmpty()
+			&& !WouldCreateBoneMatchCycle(Other, Section))
+		{
+			return true;
+		}
+	}
+	return false;
+}
+
+UClass* FRootMotionSettingsDecorationEditor::GetDecorationClass() const
+{
+	return UMovieSceneRootMotionSettingsDecoration::StaticClass();
+}
+
+void FRootMotionSettingsDecorationEditor::OnSectionInterfaceCreated(
+	UMovieSceneSection& Section,
+	UObject& Decoration,
+	TWeakPtr<ISequencer> Sequencer)
+{
+	UMovieSceneRootMotionSettingsDecoration* Settings = Cast<UMovieSceneRootMotionSettingsDecoration>(&Decoration);
+	if (!Settings)
+	{
+		return;
+	}
+
+	if (TrackedDecorations.Contains(Settings))
+	{
+		return;
+	}
+
+	Settings->BoneMatchChannel.OnKeyMovedEvent().AddRaw(this, &FRootMotionSettingsDecorationEditor::HandleBoneMatchKeyMoved);
+	Settings->BoneMatchChannel.OnKeyDeletedEvent().AddRaw(this, &FRootMotionSettingsDecorationEditor::HandleBoneMatchKeyDeleted);
+
+	TrackedDecorations.Add(Settings, { &Section, Sequencer });
+}
+
+FRootMotionSettingsDecorationEditor::FTrackedSectionInfo* FRootMotionSettingsDecorationEditor::FindTrackedInfoForChannel(
+	FMovieSceneChannel* Channel, UMovieSceneRootMotionSettingsDecoration*& OutDecoration)
+{
+	OutDecoration = nullptr;
+	for (TPair<TWeakObjectPtr<UMovieSceneRootMotionSettingsDecoration>, FTrackedSectionInfo>& Pair : TrackedDecorations)
+	{
+		UMovieSceneRootMotionSettingsDecoration* Decoration = Pair.Key.Get();
+		if (Decoration && &Decoration->BoneMatchChannel == Channel)
+		{
+			OutDecoration = Decoration;
+			return &Pair.Value;
+		}
+	}
+	return nullptr;
+}
+
+void FRootMotionSettingsDecorationEditor::UnsubscribeAll()
+{
+	for (TPair<TWeakObjectPtr<UMovieSceneRootMotionSettingsDecoration>, FTrackedSectionInfo>& Pair : TrackedDecorations)
+	{
+		if (UMovieSceneRootMotionSettingsDecoration* Decoration = Pair.Key.Get())
+		{
+			Decoration->BoneMatchChannel.OnKeyMovedEvent().RemoveAll(this);
+			Decoration->BoneMatchChannel.OnKeyDeletedEvent().RemoveAll(this);
+		}
+	}
+	TrackedDecorations.Empty();
+}
+
+void FRootMotionSettingsDecorationEditor::HandleBoneMatchKeyMoved(FMovieSceneChannel* Channel, const TArray<FKeyMoveEventItem>& Items)
+{
+	using namespace UE::MovieScene;
+
+	UMovieSceneRootMotionSettingsDecoration* Decoration = nullptr;
+	FTrackedSectionInfo* Info = FindTrackedInfoForChannel(Channel, Decoration);
+	if (!Info)
+	{
+		return;
+	}
+
+	UMovieSceneSection* Section = Info->Section.Get();
+	TSharedPtr<ISequencer> SequencerPtr = Info->Sequencer.Pin();
+	if (!Section || !SequencerPtr || !Decoration || Items.IsEmpty())
+	{
+		return;
+	}
+
+	if (!Decoration->HasBoneMatch())
+	{
+		return;
+	}
+
+	UMovieSceneAnimationMixerTrack* MixerTrack = Cast<UMovieSceneAnimationMixerTrack>(Section->GetTypedOuter<UMovieSceneTrack>());
+	if (!MixerTrack)
+	{
+		return;
+	}
+
+	TMovieSceneChannelData<FMovieSceneBoneMatchData> Data = Decoration->BoneMatchChannel.GetData();
+	TArrayView<FMovieSceneBoneMatchData> Values = Data.GetValues();
+	if (Items[0].NewIndex < 0 || Items[0].NewIndex >= Values.Num())
+	{
+		return;
+	}
+
+	FMovieSceneBoneMatchData& KeyValue = Values[Items[0].NewIndex];
+	const FFrameNumber KeyTime = Items[0].NewFrame;
+
+	// Validate: key time must be within the target section's range
+	bool bValid = Section->GetRange().Contains(KeyTime);
+
+	// Validate: the underlying mix must have at least one contributing
+	// section at the key time for the match to be meaningful.
+	// If a specific reference section is set, the key must be within
+	// its overlap with the target. Otherwise (AtCurrentTime mode),
+	// any section on the same layer or below that covers the key time
+	// counts.
+	if (bValid)
+	{
+		UMovieSceneSection* Ref = KeyValue.ReferenceSection.Get();
+		if (Ref)
+		{
+			TRange<FFrameNumber> Overlap = TRange<FFrameNumber>::Intersection(Section->GetRange(), Ref->GetRange());
+			if (Overlap.IsEmpty() || !Overlap.Contains(KeyTime))
+			{
+				bValid = false;
+			}
+		}
+		else
+		{
+			const int32 TargetRow = Section->GetRowIndex();
+			bool bHasUnderlyingPose = false;
+			for (UMovieSceneSection* Other : MixerTrack->GetAllSections())
+			{
+				if (Other && Other != Section
+					&& !Other->IsA<UMovieSceneAnimTransitionSectionBase>()
+					&& Other->GetRowIndex() <= TargetRow
+					&& Other->GetRange().Contains(KeyTime))
+				{
+					bHasUnderlyingPose = true;
+					break;
+				}
+			}
+			if (!bHasUnderlyingPose)
+			{
+				bValid = false;
+			}
+		}
+	}
+
+	if (!bValid)
+	{
+		KeyValue.bIsValid = false;
+		KeyValue.bIsDirty = true;
+		SequencerPtr->NotifyMovieSceneDataChanged(EMovieSceneDataChangeType::TrackValueChanged);
+		return;
+	}
+
+	AnimMixerBoneMatching::FBoneMatchingContext Context;
+	Context.Linker = SequencerPtr->GetEvaluationTemplate().GetEntitySystemLinker();
+	Context.InstanceHandle = SequencerPtr->GetEvaluationTemplate().GetRootInstanceHandle();
+	Context.MixerTrack = MixerTrack;
+	Context.CurrentTime = FFrameTime(KeyTime);
+
+	FMovieSceneBoneMatchData Result = AnimMixerBoneMatching::ComputeBoneMatch(
+		Section, KeyValue, Context);
+
+	KeyValue = Result;
+
+	// Propagate rematches through dependency chain (A->B->C)
+	if (Result.bIsValid)
+	{
+		TSet<UMovieSceneSection*> Visited;
+		Visited.Add(Section);
+		AnimMixerBoneMatching::PropagateRematch(Section, Context, Visited);
+	}
+
+	SequencerPtr->NotifyMovieSceneDataChanged(EMovieSceneDataChangeType::TrackValueChanged);
+}
+
+void FRootMotionSettingsDecorationEditor::HandleBoneMatchKeyDeleted(FMovieSceneChannel* Channel, const TArray<FKeyAddOrDeleteEventItem>& Items)
+{
+	UMovieSceneRootMotionSettingsDecoration* Decoration = nullptr;
+	FTrackedSectionInfo* Info = FindTrackedInfoForChannel(Channel, Decoration);
+	if (!Info)
+	{
+		return;
+	}
+
+	UMovieSceneSection* Section = Info->Section.Get();
+	TSharedPtr<ISequencer> SequencerPtr = Info->Sequencer.Pin();
+	if (!Section || !SequencerPtr)
+	{
+		return;
+	}
+
+	// Hide the channel row now that there are no keys
+	Section->InvalidateChannelProxy();
+	SequencerPtr->NotifyMovieSceneDataChanged(EMovieSceneDataChangeType::TrackValueChanged);
+}
+
+void FRootMotionSettingsDecorationEditor::BuildSectionContextMenu(
+	FMenuBuilder& MenuBuilder,
+	UMovieSceneSection& Section,
+	UObject& Decoration,
+	const FGuid& ObjectBinding,
+	TWeakPtr<ISequencer> Sequencer)
+{
+	UMovieSceneRootMotionSettingsDecoration* Settings = Cast<UMovieSceneRootMotionSettingsDecoration>(&Decoration);
+	if (!Settings)
+	{
+		return;
+	}
+
+	MenuBuilder.AddSubMenu(
+		LOCTEXT("RootTransform_Label", "Root Transform"),
+		LOCTEXT("RootTransform_Tooltip", "Options for root transform behavior from this section"),
+		FNewMenuDelegate::CreateRaw(this, &FRootMotionSettingsDecorationEditor::PopulateRootTransformMenu,
+			TWeakObjectPtr<UMovieSceneSection>(&Section),
+			TWeakObjectPtr<UMovieSceneRootMotionSettingsDecoration>(Settings),
+			ObjectBinding,
+			Sequencer)
+	);
+}
+
+void FRootMotionSettingsDecorationEditor::PopulateRootTransformMenu(
+	FMenuBuilder& MenuBuilder,
+	TWeakObjectPtr<UMovieSceneSection> WeakSection,
+	TWeakObjectPtr<UMovieSceneRootMotionSettingsDecoration> WeakSettings,
+	FGuid ObjectBinding,
+	TWeakPtr<ISequencer> Sequencer)
+{
+	UMovieSceneSection* Section = WeakSection.Get();
+	UMovieSceneRootMotionSettingsDecoration* Settings = WeakSettings.Get();
+	if (!Section || !Settings)
+	{
+		return;
+	}
+
+	const bool bBoneMatchInvalid = Settings->HasBoneMatch() && !Settings->GetBoneMatchData().bIsValid;
+
+	MenuBuilder.BeginSection(NAME_None, LOCTEXT("RootTransformModeSection_Label", "Root Transform Mode"));
+	{
+		FCanExecuteAction AlwaysExecute = FCanExecuteAction::CreateLambda([] { return true; });
+
+		auto IsMode = [WeakSettings](EMovieSceneRootMotionTransformMode Mode)
+		{
+			if (UMovieSceneRootMotionSettingsDecoration* Settings = WeakSettings.Get())
+			{
+				return Settings->GetRootTransformMode() == Mode;
+			}
+			return Mode == EMovieSceneRootMotionTransformMode::Asset;
+		};
+
+		auto SetMode = [WeakSettings, WeakSection, Sequencer](EMovieSceneRootMotionTransformMode Mode)
+		{
+			UMovieSceneRootMotionSettingsDecoration* Settings = WeakSettings.Get();
+			UMovieSceneSection* Section = WeakSection.Get();
+			if (Settings && Section)
+			{
+				Settings->Modify();
+				Settings->TransformMode.SetDefault((uint8)Mode);
+				Settings->InvalidateDecorationChannelProxy();
+				if (TSharedPtr<ISequencer> SequencerPtr = Sequencer.Pin())
+				{
+					SequencerPtr->NotifyMovieSceneDataChanged(EMovieSceneDataChangeType::MovieSceneStructureItemsChanged);
+				}
+			}
+		};
+
+		auto AddModeEntry = [&](FText Label, FText Tooltip, FText TxnLabel, EMovieSceneRootMotionTransformMode Mode)
+		{
+			MenuBuilder.AddMenuEntry(
+				Label, Tooltip, FSlateIcon(),
+				FUIAction(
+					FExecuteAction::CreateLambda([SetMode, TxnLabel, Mode]{ FScopedTransaction Txn(TxnLabel); SetMode(Mode); }),
+					AlwaysExecute,
+					FIsActionChecked::CreateLambda([IsMode, Mode]{ return IsMode(Mode); })
+				),
+				NAME_None, EUserInterfaceActionType::RadioButton);
+		};
+
+		AddModeEntry(
+			LOCTEXT("RootTransformMode_None", "From Animation"),
+			LOCTEXT("RootTransformMode_None_Tip", "Use the root motion transform directly from the animation asset"),
+			LOCTEXT("SetRootTransformNone", "Use Root Transform from Asset"),
+			EMovieSceneRootMotionTransformMode::Asset);
+
+		AddModeEntry(
+			LOCTEXT("RootTransformMode_Offset", "Offset From Animation"),
+			LOCTEXT("RootTransformMode_Offset_Tooltip", "Offset the root motion transform using keyframed values"),
+			LOCTEXT("SetRootTransformOffset", "Offset Root Transform from Asset"),
+			EMovieSceneRootMotionTransformMode::Offset);
+
+		AddModeEntry(
+			LOCTEXT("RootTransformMode_Override", "Manual Override"),
+			LOCTEXT("RootTransformMode_Override_Tooltip", "Completely override the root motion transform using keyframed values"),
+			LOCTEXT("SetRootTransformOverride", "Override Root Transform"),
+			EMovieSceneRootMotionTransformMode::Override);
+
+		AddModeEntry(
+			LOCTEXT("RootTransformMode_AccumulatedOffset", "Accumulated Offset"),
+			LOCTEXT("RootTransformMode_AccumulatedOffset_Tooltip", "Automatically chain from previous KeepState sections. The section starts where the previous section ended."),
+			LOCTEXT("SetRootTransformAccumulatedOffset", "Use Accumulated Offset"),
+			EMovieSceneRootMotionTransformMode::AccumulatedOffset);
+
+		// Bone Match is display-only (set via the Match Bone dialog, not directly selectable)
+		MenuBuilder.AddMenuEntry(
+			bBoneMatchInvalid
+				? LOCTEXT("RootTransformMode_BoneMatchInvalid", "Bone Match (invalid)")
+				: LOCTEXT("RootTransformMode_BoneMatch", "Bone Match"),
+			LOCTEXT("RootTransformMode_BoneMatch_Tooltip", "Position based on bone alignment. Set automatically by the Match Bone dialog."),
+			FSlateIcon(),
+			FUIAction(
+				FExecuteAction::CreateLambda([]{}),
+				FCanExecuteAction::CreateLambda([] { return false; }),
+				FIsActionChecked::CreateLambda([IsMode]{ return IsMode(EMovieSceneRootMotionTransformMode::BoneMatch); })
+			),
+			NAME_None,
+			EUserInterfaceActionType::RadioButton
+		);
+	}
+	MenuBuilder.EndSection();
+
+	// Re-center option (only in Offset mode)
+	if (Settings->GetRootTransformMode() == EMovieSceneRootMotionTransformMode::Offset)
+	{
+		MenuBuilder.BeginSection(NAME_None, LOCTEXT("RootTransformOffset_Label", "Offset"));
+		{
+			MenuBuilder.AddMenuEntry(
+				LOCTEXT("RecenterRootTransform", "Re-center Root Transform"),
+				LOCTEXT("RecenterRootTransform_Tooltip", "Center the root transform for this animation around its current position"),
+				FSlateIcon(),
+				FUIAction(
+					FExecuteAction::CreateRaw(this, &FRootMotionSettingsDecorationEditor::RecenterRootTransform,
+						WeakSection.Get(), WeakSettings.Get(), Sequencer)
+				)
+			);
+		}
+		MenuBuilder.EndSection();
+	}
+
+	const bool bHasReferenceSections = HasOverlappingReferenceSections(Section);
+
+	// Bone Matching
+	MenuBuilder.BeginSection(NAME_None,
+		bBoneMatchInvalid
+			? LOCTEXT("BoneMatchingInvalid_Label", "Bone Match (invalid)")
+			: LOCTEXT("BoneMatching_Label", "Bone Matching"));
+	{
+		if (Settings->HasBoneMatch())
+		{
+			// Re-match: directly re-compute with existing settings (no dialog)
+			auto GetMatchTimeModeLabel = [](EBoneMatchTimeMode Mode) -> FText
+			{
+				switch (Mode)
+				{
+				case EBoneMatchTimeMode::AtCurrentTime:            return LOCTEXT("TimeMode_Current", "Current Time");
+				case EBoneMatchTimeMode::AtStartOfSelectedSection: return LOCTEXT("TimeMode_StartSel", "Start of This Section");
+				case EBoneMatchTimeMode::AtEndOfSelectedSection:   return LOCTEXT("TimeMode_EndSel", "End of This Section");
+				case EBoneMatchTimeMode::AtStartOfReferenceSection:return LOCTEXT("TimeMode_StartRef", "Start of Reference");
+				case EBoneMatchTimeMode::AtEndOfReferenceSection:  return LOCTEXT("TimeMode_EndRef", "End of Reference");
+				case EBoneMatchTimeMode::InBetween:                return LOCTEXT("TimeMode_Between", "Midpoint");
+				default:                                           return LOCTEXT("TimeMode_Unknown", "Unknown");
+				}
+			};
+
+			FMovieSceneBoneMatchData CurrentMatch = Settings->GetBoneMatchData();
+			FText RematchLabel = FText::Format(LOCTEXT("RematchBone", "Re-match ({0})"),
+				FText::FromName(CurrentMatch.BoneName));
+			FText RematchTooltip = FText::Format(LOCTEXT("RematchBone_Tooltip", "Re-compute bone match: {0}, time: {1}"),
+				FText::FromName(CurrentMatch.BoneName),
+				GetMatchTimeModeLabel(CurrentMatch.MatchTimeMode));
+
+			// Re-match and edit are only available when there's a valid
+			// reference section to match against
+			const bool bCanRematch = bHasReferenceSections;
+			FText DisabledTooltip = LOCTEXT("BoneMatch_NoOverlap_Tooltip", "No overlapping root-motion-producing sections available for bone matching");
+
+			MenuBuilder.AddMenuEntry(
+				RematchLabel,
+				bCanRematch ? RematchTooltip : DisabledTooltip,
+				FSlateIcon(),
+				FUIAction(
+					FExecuteAction::CreateRaw(this, &FRootMotionSettingsDecorationEditor::RematchBone,
+						WeakSection.Get(), WeakSettings.Get(), Sequencer),
+					FCanExecuteAction::CreateLambda([bCanRematch] { return bCanRematch; })
+				)
+			);
+
+			// Edit: open the dialog to change match parameters
+			MenuBuilder.AddMenuEntry(
+				LOCTEXT("EditBoneMatch", "Edit Bone Match..."),
+				bCanRematch
+					? LOCTEXT("EditBoneMatch_Tooltip", "Open dialog to change bone match settings")
+					: DisabledTooltip,
+				FSlateIcon(),
+				FUIAction(
+					FExecuteAction::CreateRaw(this, &FRootMotionSettingsDecorationEditor::OpenBoneMatchDialog,
+						WeakSection.Get(), WeakSettings.Get(), ObjectBinding, Sequencer),
+					FCanExecuteAction::CreateLambda([bCanRematch] { return bCanRematch; })
+				)
+			);
+
+			MenuBuilder.AddMenuEntry(
+				LOCTEXT("ClearBoneMatch", "Clear Bone Match"),
+				LOCTEXT("ClearBoneMatch_Tooltip", "Remove the bone match transform offset and revert to previous transform mode"),
+				FSlateIcon(),
+				FUIAction(
+					FExecuteAction::CreateRaw(this, &FRootMotionSettingsDecorationEditor::ClearBoneMatch,
+						WeakSection.Get(), WeakSettings.Get(), Sequencer)
+				)
+			);
+		}
+		else
+		{
+			// No existing match: open dialog
+			MenuBuilder.AddMenuEntry(
+				LOCTEXT("MatchBone", "Match Bone..."),
+				bHasReferenceSections
+					? LOCTEXT("MatchBone_Tooltip", "Open bone match dialog to align this section's root transform with a neighboring section")
+					: LOCTEXT("MatchBone_NoRef_Tooltip", "No other sections available for bone matching"),
+				FSlateIcon(),
+				FUIAction(
+					FExecuteAction::CreateRaw(this, &FRootMotionSettingsDecorationEditor::OpenBoneMatchDialog,
+						WeakSection.Get(), WeakSettings.Get(), ObjectBinding, Sequencer),
+					FCanExecuteAction::CreateLambda([bHasReferenceSections] { return bHasReferenceSections; })
+				)
+			);
+		}
+	}
+	MenuBuilder.EndSection();
+}
+
+void FRootMotionSettingsDecorationEditor::RecenterRootTransform(
+	UMovieSceneSection* Section,
+	UMovieSceneRootMotionSettingsDecoration* Settings,
+	TWeakPtr<ISequencer> WeakSequencer)
+{
+	using namespace UE::MovieScene;
+
+	TSharedPtr<ISequencer> SequencerPtr = WeakSequencer.Pin();
+	if (!SequencerPtr || !Section || !Settings)
+	{
+		return;
+	}
+
+	UMovieSceneAnimationMixerTrack* MixerTrack = Cast<UMovieSceneAnimationMixerTrack>(Section->GetTypedOuter<UMovieSceneTrack>());
+	UMovieSceneEntitySystemLinker* Linker = SequencerPtr->GetEvaluationTemplate().GetEntitySystemLinker();
+	if (!MixerTrack || !Linker)
+	{
+		return;
+	}
+
+	// Evaluate this section in isolation at the current time,
+	// skipping root motion space conversion to get the animation-space transform
+	FInstanceHandle InstanceHandle = SequencerPtr->GetEvaluationTemplate().GetRootInstanceHandle();
+	FQualifiedFrameTime Time = SequencerPtr->GetLocalTime();
+
+	AnimMixerBakeEvaluation::FBakeFilter Filter;
+	Filter.IncludeOnlySections.Add(FObjectKey(Section));
+	Filter.bSkipRootMotionConversion = true;
+
+	AnimMixerBakeEvaluation::FBakeResult Result =
+		AnimMixerBakeEvaluation::EvaluateAtTime(
+			Linker, InstanceHandle, MixerTrack,
+			Time.Time, Filter);
+
+	FScopedTransaction Transaction(LOCTEXT("RecenterRootTransform", "Re-center Root Transform"));
+	Settings->Modify();
+	Settings->RootOriginLocation = Result.RootMotionTransform.GetTranslation();
+}
+
+void FRootMotionSettingsDecorationEditor::OpenBoneMatchDialog(
+	UMovieSceneSection* Section,
+	UMovieSceneRootMotionSettingsDecoration* Settings,
+	FGuid ObjectBinding,
+	TWeakPtr<ISequencer> WeakSequencer)
+{
+	TSharedPtr<ISequencer> SequencerPtr = WeakSequencer.Pin();
+	if (!Section || !SequencerPtr || !Settings)
+	{
+		return;
+	}
+
+	UMovieSceneAnimationMixerTrack* MixerTrack = Cast<UMovieSceneAnimationMixerTrack>(Section->GetTypedOuter<UMovieSceneTrack>());
+	if (!MixerTrack)
+	{
+		return;
+	}
+
+	TSharedRef<SBoneMatchDialog> Dialog = SNew(SBoneMatchDialog)
+		.Sequencer(SequencerPtr.Get())
+		.TargetSection(Section)
+		.MixerTrack(MixerTrack)
+		.Decoration(Settings)
+		.ObjectBinding(ObjectBinding)
+		.OnAccepted(FOnBoneMatchAccepted::CreateRaw(
+			this, &FRootMotionSettingsDecorationEditor::OnBoneMatchAccepted,
+			Section, Settings, WeakSequencer));
+
+	Dialog->OpenDialog(true);
+}
+
+void FRootMotionSettingsDecorationEditor::OnBoneMatchAccepted(
+	const FMovieSceneBoneMatchData& MatchSettings,
+	UMovieSceneSection* Section,
+	UMovieSceneRootMotionSettingsDecoration* Settings,
+	TWeakPtr<ISequencer> WeakSequencer)
+{
+	using namespace UE::MovieScene;
+
+	TSharedPtr<ISequencer> SequencerPtr = WeakSequencer.Pin();
+	if (!Section || !SequencerPtr || !Settings)
+	{
+		return;
+	}
+
+	UMovieSceneAnimationMixerTrack* MixerTrack = Cast<UMovieSceneAnimationMixerTrack>(Section->GetTypedOuter<UMovieSceneTrack>());
+	if (!MixerTrack)
+	{
+		return;
+	}
+
+	AnimMixerBoneMatching::FBoneMatchingContext Context;
+	Context.Linker = SequencerPtr->GetEvaluationTemplate().GetEntitySystemLinker();
+	Context.InstanceHandle = SequencerPtr->GetEvaluationTemplate().GetRootInstanceHandle();
+	Context.MixerTrack = MixerTrack;
+	Context.CurrentTime = SequencerPtr->GetLocalTime().Time;
+
+	FScopedTransaction Transaction(LOCTEXT("BoneMatch", "Bone Match"));
+	Settings->Modify();
+
+	FMovieSceneBoneMatchData Result = AnimMixerBoneMatching::ComputeBoneMatch(
+		Section, MatchSettings, Context);
+
+	// Store the result as a key on the channel at the resolved match time
+	Settings->BoneMatchChannel.Reset();
+	if (Result.bIsValid)
+	{
+		Settings->BoneMatchChannel.GetData().AddKey(Result.MatchTime, Result);
+		Settings->TransformMode.SetDefault((uint8)EMovieSceneRootMotionTransformMode::BoneMatch);
+
+		// Propagate rematches through dependency chain
+		TSet<UMovieSceneSection*> Visited;
+		Visited.Add(Section);
+		AnimMixerBoneMatching::PropagateRematch(Section, Context, Visited);
+	}
+	Section->InvalidateChannelProxy();
+
+	SequencerPtr->NotifyMovieSceneDataChanged(EMovieSceneDataChangeType::TrackValueChanged);
+}
+
+void FRootMotionSettingsDecorationEditor::RematchBone(
+	UMovieSceneSection* Section,
+	UMovieSceneRootMotionSettingsDecoration* Settings,
+	TWeakPtr<ISequencer> WeakSequencer)
+{
+	if (!Section || !Settings || !Settings->HasBoneMatch())
+	{
+		return;
+	}
+
+	OnBoneMatchAccepted(Settings->GetBoneMatchData(), Section, Settings, WeakSequencer);
+}
+
+void FRootMotionSettingsDecorationEditor::ClearBoneMatch(
+	UMovieSceneSection* Section,
+	UMovieSceneRootMotionSettingsDecoration* Settings,
+	TWeakPtr<ISequencer> WeakSequencer)
+{
+	TSharedPtr<ISequencer> SequencerPtr = WeakSequencer.Pin();
+	if (!SequencerPtr || !Section || !Settings)
+	{
+		return;
+	}
+
+	FScopedTransaction Transaction(LOCTEXT("ClearBoneMatch", "Clear Bone Match"));
+	Settings->Modify();
+	Settings->BoneMatchChannel.Reset();
+
+	// Revert transform mode if it was set to BoneMatch
+	if (Settings->GetRootTransformMode() == EMovieSceneRootMotionTransformMode::BoneMatch)
+	{
+		Settings->TransformMode.SetDefault((uint8)EMovieSceneRootMotionTransformMode::Asset);
+	}
+	Section->InvalidateChannelProxy();
+
+	SequencerPtr->NotifyMovieSceneDataChanged(EMovieSceneDataChangeType::TrackValueChanged);
+}
+
+TSharedPtr<SWidget> FRootMotionSettingsDecorationEditor::CreateKeyFrameWidget(
+	UObject& Decoration,
+	UMovieSceneSection& Section,
+	FMovieSceneChannelHandle ChannelHandle,
+	TWeakPtr<ISequencer> Sequencer,
+	const FGuid& ObjectBindingID,
+	TSharedPtr<FViewModel> ChannelGroupModel)
+{
+	UMovieSceneRootMotionSettingsDecoration* Settings = Cast<UMovieSceneRootMotionSettingsDecoration>(&Decoration);
+	if (!Settings)
+	{
+		return nullptr;
+	}
+
+	// Only customize the BoneMatchChannel
+	FMovieSceneChannel* RawChannel = ChannelHandle.Get();
+	if (RawChannel != static_cast<FMovieSceneChannel*>(&Settings->BoneMatchChannel))
+	{
+		return nullptr;
+	}
+
+	TWeakObjectPtr<UMovieSceneSection> WeakSection = &Section;
+	TWeakObjectPtr<UMovieSceneRootMotionSettingsDecoration> WeakSettings = Settings;
+	TWeakPtr<ISequencer> WeakSequencer = Sequencer;
+	FGuid ObjectBinding = ObjectBindingID;
+
+	auto CanMatchBone = [WeakSection]() -> bool
+	{
+		UMovieSceneSection* Section = WeakSection.Get();
+		return Section && HasOverlappingReferenceSections(Section);
+	};
+
+	FText Tooltip = Settings->HasBoneMatch()
+		? LOCTEXT("EditBoneMatch_KeyFrame_Tooltip", "Edit bone match settings")
+		: LOCTEXT("MatchBone_KeyFrame_Tooltip", "Open bone match dialog");
+
+	return SNew(SKeyNavigationButtons, ChannelGroupModel)
+		.Buttons(EKeyNavigationButtons::AddKey)
+		.AddKeyToolTip(Tooltip)
+		.IsEnabled_Lambda(CanMatchBone)
+		.OnAddKey(SKeyNavigationButtons::FOnAddKey::CreateLambda(
+			[this, WeakSection, WeakSettings, ObjectBinding, WeakSequencer](FFrameTime, TSharedPtr<FViewModel>)
+			{
+				UMovieSceneSection* Section = WeakSection.Get();
+				UMovieSceneRootMotionSettingsDecoration* Settings = WeakSettings.Get();
+				if (Section && Settings)
+				{
+					OpenBoneMatchDialog(Section, Settings, ObjectBinding, WeakSequencer);
+				}
+			}));
+}
+
+TSharedPtr<SWidget> FRootMotionSettingsDecorationEditor::CreateKeyEditor(
+	UObject& Decoration,
+	UMovieSceneSection& Section,
+	FMovieSceneChannelHandle ChannelHandle,
+	TWeakPtr<ISequencer> Sequencer,
+	const FGuid& ObjectBindingID)
+{
+	UMovieSceneRootMotionSettingsDecoration* Settings = Cast<UMovieSceneRootMotionSettingsDecoration>(&Decoration);
+	if (!Settings)
+	{
+		return nullptr;
+	}
+
+	// Only customize the TransformMode channel
+	FMovieSceneChannel* RawChannel = ChannelHandle.Get();
+	if (RawChannel != static_cast<FMovieSceneChannel*>(&Settings->TransformMode))
+	{
+		return nullptr;
+	}
+
+	UEnum* Enum = Settings->TransformMode.GetEnum();
+	if (!Enum)
+	{
+		return nullptr;
+	}
+
+	TWeakObjectPtr<UMovieSceneRootMotionSettingsDecoration> WeakSettings = Settings;
+	TWeakObjectPtr<UMovieSceneSection> WeakSection = &Section;
+	TWeakPtr<ISequencer> WeakSequencer = Sequencer;
+
+	return MovieSceneToolHelpers::MakeEnumComboBox(
+		Enum,
+		TAttribute<int32>::CreateLambda([WeakSettings]() -> int32
+		{
+			if (UMovieSceneRootMotionSettingsDecoration* Settings = WeakSettings.Get())
+			{
+				return (int32)Settings->TransformMode.GetDefault().Get(0);
+			}
+			return 0;
+		}),
+		SEnumComboBox::FOnEnumSelectionChanged::CreateLambda(
+			[WeakSettings, WeakSection, WeakSequencer](int32 Selection, ESelectInfo::Type)
+			{
+				UMovieSceneRootMotionSettingsDecoration* Settings = WeakSettings.Get();
+				UMovieSceneSection* Section = WeakSection.Get();
+				if (!Settings || !Section)
+				{
+					return;
+				}
+
+				FScopedTransaction Transaction(LOCTEXT("SetTransformMode", "Set Transform Mode"));
+				Settings->Modify();
+				Settings->TransformMode.SetDefault((uint8)Selection);
+				Settings->InvalidateDecorationChannelProxy();
+				if (TSharedPtr<ISequencer> SequencerPtr = WeakSequencer.Pin())
+				{
+					SequencerPtr->NotifyMovieSceneDataChanged(EMovieSceneDataChangeType::MovieSceneStructureItemsChanged);
+				}
+			})
+	);
+}
+
+} // namespace UE::Sequencer
+
+#undef LOCTEXT_NAMESPACE
diff --git a/Engine/Plugins/Experimental/MovieSceneAnimMixer/Source/MovieSceneAnimMixerEditor/Private/RootMotionSettingsDecorationEditor.h b/Engine/Plugins/Experimental/MovieSceneAnimMixer/Source/MovieSceneAnimMixerEditor/Private/RootMotionSettingsDecorationEditor.h
new file mode 100644
--- /dev/null
+++ b/Engine/Plugins/Experimental/MovieSceneAnimMixer/Source/MovieSceneAnimMixerEditor/Private/RootMotionSettingsDecorationEditor.h
@@ -0,0 +1,105 @@
+// Copyright Epic Games, Inc. All Rights Reserved.
+
+#pragma once
+
+#include "ISequencerDecorationEditor.h"
+#include "ISequencerSection.h"
+#include "Channels/MovieSceneChannel.h"
+#include "UObject/WeakObjectPtr.h"
+
+class UMovieSceneRootMotionSettingsDecoration;
+struct FMovieSceneBoneMatchData;
+
+namespace UE::Sequencer
+{
+
+// Decoration editor for UMovieSceneRootMotionSettingsDecoration.
+// Provides the "Root Transform" right-click context menu with transform mode
+// selection, re-center, and bone matching.
+class FRootMotionSettingsDecorationEditor : public ISequencerDecorationEditor
+{
+public:
+	virtual UClass* GetDecorationClass() const override;
+
+	virtual void OnSectionInterfaceCreated(
+		UMovieSceneSection& Section,
+		UObject& Decoration,
+		TWeakPtr<ISequencer> Sequencer) override;
+
+	virtual void BuildSectionContextMenu(FMenuBuilder& MenuBuilder,
+		UMovieSceneSection& Section,
+		UObject& Decoration,
+		const FGuid& ObjectBinding,
+		TWeakPtr<ISequencer> Sequencer) override;
+
+	virtual TSharedPtr<SWidget> CreateKeyEditor(
+		UObject& Decoration,
+		UMovieSceneSection& Section,
+		FMovieSceneChannelHandle ChannelHandle,
+		TWeakPtr<ISequencer> Sequencer,
+		const FGuid& ObjectBindingID) override;
+
+	virtual TSharedPtr<SWidget> CreateKeyFrameWidget(
+		UObject& Decoration,
+		UMovieSceneSection& Section,
+		FMovieSceneChannelHandle ChannelHandle,
+		TWeakPtr<ISequencer> Sequencer,
+		const FGuid& ObjectBindingID,
+		TSharedPtr<UE::Sequencer::FViewModel> ChannelGroupModel) override;
+
+	// Unsubscribe from all tracked bone match channel events. Call during module shutdown.
+	void UnsubscribeAll();
+
+private:
+	void PopulateRootTransformMenu(
+		FMenuBuilder& MenuBuilder,
+		TWeakObjectPtr<UMovieSceneSection> Section,
+		TWeakObjectPtr<UMovieSceneRootMotionSettingsDecoration> Settings,
+		FGuid ObjectBinding,
+		TWeakPtr<ISequencer> Sequencer);
+
+	void RecenterRootTransform(
+		UMovieSceneSection* Section,
+		UMovieSceneRootMotionSettingsDecoration* Settings,
+		TWeakPtr<ISequencer> Sequencer);
+
+	void OpenBoneMatchDialog(
+		UMovieSceneSection* Section,
+		UMovieSceneRootMotionSettingsDecoration* Settings,
+		FGuid ObjectBinding,
+		TWeakPtr<ISequencer> Sequencer);
+
+	void OnBoneMatchAccepted(
+		const FMovieSceneBoneMatchData& MatchSettings,
+		UMovieSceneSection* Section,
+		UMovieSceneRootMotionSettingsDecoration* Settings,
+		TWeakPtr<ISequencer> Sequencer);
+
+	void RematchBone(
+		UMovieSceneSection* Section,
+		UMovieSceneRootMotionSettingsDecoration* Settings,
+		TWeakPtr<ISequencer> Sequencer);
+
+	void ClearBoneMatch(
+		UMovieSceneSection* Section,
+		UMovieSceneRootMotionSettingsDecoration* Settings,
+		TWeakPtr<ISequencer> Sequencer);
+
+	// Bone match channel event handlers
+	void HandleBoneMatchKeyMoved(FMovieSceneChannel* Channel, const TArray<FKeyMoveEventItem>& Items);
+	void HandleBoneMatchKeyDeleted(FMovieSceneChannel* Channel, const TArray<FKeyAddOrDeleteEventItem>& Items);
+
+	struct FTrackedSectionInfo
+	{
+		TWeakObjectPtr<UMovieSceneSection> Section;
+		TWeakPtr<ISequencer> Sequencer;
+	};
+
+	// Find the tracked entry whose decoration's BoneMatchChannel matches the
+	// event channel. Returns the info and sets OutDecoration to the live decoration.
+	FTrackedSectionInfo* FindTrackedInfoForChannel(FMovieSceneChannel* Channel, UMovieSceneRootMotionSettingsDecoration*& OutDecoration);
+
+	TMap<TWeakObjectPtr<UMovieSceneRootMotionSettingsDecoration>, FTrackedSectionInfo> TrackedDecorations;
+};
+
+} // namespace UE::Sequencer
diff --git a/Engine/Plugins/Experimental/MovieSceneAnimMixer/Source/MovieSceneAnimMixerEditor/Private/RootMotionTargetDecorationEditor.cpp b/Engine/Plugins/Experimental/MovieSceneAnimMixer/Source/MovieSceneAnimMixerEditor/Private/RootMotionTargetDecorationEditor.cpp
--- a/Engine/Plugins/Experimental/MovieSceneAnimMixer/Source/MovieSceneAnimMixerEditor/Private/RootMotionTargetDecorationEditor.cpp
+++ b/Engine/Plugins/Experimental/MovieSceneAnimMixer/Source/MovieSceneAnimMixerEditor/Private/RootMotionTargetDecorationEditor.cpp
@@ -53,7 +53,7 @@ UClass* FRootMotionTargetDecorationEditor::GetDecorationClass() const
 	return UMovieSceneRootMotionTargetDecoration::StaticClass();
 }
 
-TSharedRef<ISequencerSection> FRootMotionTargetDecorationEditor::MakeSectionInterface(
+TSharedPtr<ISequencerSection> FRootMotionTargetDecorationEditor::MakeSectionInterface(
 	UMovieSceneSection& Section,
 	UMovieSceneTrack& OwningTrack,
 	FGuid ObjectBinding)
diff --git a/Engine/Plugins/Experimental/MovieSceneAnimMixer/Source/MovieSceneAnimMixerEditor/Private/RootMotionTargetDecorationEditor.h b/Engine/Plugins/Experimental/MovieSceneAnimMixer/Source/MovieSceneAnimMixerEditor/Private/RootMotionTargetDecorationEditor.h
--- a/Engine/Plugins/Experimental/MovieSceneAnimMixer/Source/MovieSceneAnimMixerEditor/Private/RootMotionTargetDecorationEditor.h
+++ b/Engine/Plugins/Experimental/MovieSceneAnimMixer/Source/MovieSceneAnimMixerEditor/Private/RootMotionTargetDecorationEditor.h
@@ -42,7 +42,7 @@ class FRootMotionTargetDecorationEditor : public ISequencerDecorationEditor
 {
 public:
 	virtual UClass* GetDecorationClass() const override;
-	virtual TSharedRef<ISequencerSection> MakeSectionInterface(
+	virtual TSharedPtr<ISequencerSection> MakeSectionInterface(
 		UMovieSceneSection& Section,
 		UMovieSceneTrack& OwningTrack,
 		FGuid ObjectBinding) override;
diff --git a/Engine/Plugins/Experimental/MovieSceneAnimMixer/Source/MovieSceneAnimMixerEditor/Private/SBoneMatchDialog.cpp b/Engine/Plugins/Experimental/MovieSceneAnimMixer/Source/MovieSceneAnimMixerEditor/Private/SBoneMatchDialog.cpp
new file mode 100644
--- /dev/null
+++ b/Engine/Plugins/Experimental/MovieSceneAnimMixer/Source/MovieSceneAnimMixerEditor/Private/SBoneMatchDialog.cpp
@@ -0,0 +1,584 @@
+// Copyright Epic Games, Inc. All Rights Reserved.
+
+#include "SBoneMatchDialog.h"
+
+#include "BoneSelectionWidget.h"
+#include "Components/SkeletalMeshComponent.h"
+#include "Engine/SkeletalMesh.h"
+#include "Editor.h"
+#include "Framework/Application/SlateApplication.h"
+#include "ISequencer.h"
+#include "ISequencerSection.h"
+#include "MovieSceneAnimationMixerTrack.h"
+#include "MovieSceneAnimTransitionSectionBase.h"
+#include "MovieSceneSection.h"
+#include "MVVM/SectionModelStorageExtension.h"
+#include "MVVM/ViewModels/SectionModel.h"
+#include "MVVM/ViewModels/SequencerEditorViewModel.h"
+#include "Widgets/Input/SButton.h"
+#include "Widgets/Input/SCheckBox.h"
+#include "Widgets/Input/SComboBox.h"
+#include "Widgets/Layout/SSpacer.h"
+#include "Widgets/SWindow.h"
+#include "Widgets/Text/STextBlock.h"
+
+#define LOCTEXT_NAMESPACE "SBoneMatchDialog"
+
+static USkeletalMeshComponent* FindBoundSkeletalMeshComponent(ISequencer* Sequencer, const FGuid& ObjectBinding)
+{
+	if (!Sequencer || !ObjectBinding.IsValid())
+	{
+		return nullptr;
+	}
+
+	TArrayView<TWeakObjectPtr<>> BoundObjects = Sequencer->FindBoundObjects(ObjectBinding, Sequencer->GetFocusedTemplateID());
+	for (const TWeakObjectPtr<>& WeakObj : BoundObjects)
+	{
+		UObject* Obj = WeakObj.Get();
+		if (!Obj)
+		{
+			continue;
+		}
+
+		USkeletalMeshComponent* SkelComp = nullptr;
+		if (AActor* Actor = Cast<AActor>(Obj))
+		{
+			SkelComp = Actor->FindComponentByClass<USkeletalMeshComponent>();
+		}
+		else
+		{
+			SkelComp = Cast<USkeletalMeshComponent>(Obj);
+		}
+
+		if (SkelComp && SkelComp->GetSkeletalMeshAsset())
+		{
+			return SkelComp;
+		}
+	}
+	return nullptr;
+}
+
+static FString GetSectionDisplayName(const UMovieSceneSection* Section, ISequencer* Sequencer)
+{
+	using namespace UE::Sequencer;
+
+	if (Sequencer)
+	{
+		if (TSharedPtr<FSequencerEditorViewModel> ViewModel = Sequencer->GetViewModel())
+		{
+			if (FSectionModelStorageExtension* Storage = ViewModel->GetRootModel()->CastDynamic<FSectionModelStorageExtension>())
+			{
+				if (TSharedPtr<FSectionModel> Model = Storage->FindModelForSection(Section))
+				{
+					if (TSharedPtr<ISequencerSection> SectionInterface = Model->GetSectionInterface())
+					{
+						FText Title = SectionInterface->GetSectionTitle();
+						if (!Title.IsEmpty())
+						{
+							return Title.ToString();
+						}
+					}
+				}
+			}
+		}
+	}
+	return Section->GetName();
+}
+
+void SBoneMatchDialog::Construct(const FArguments& InArgs)
+{
+	Sequencer = InArgs._Sequencer;
+	TargetSection = InArgs._TargetSection;
+	MixerTrack = InArgs._MixerTrack;
+	Decoration = InArgs._Decoration;
+	ObjectBinding = InArgs._ObjectBinding;
+	OnAccepted = InArgs._OnAccepted;
+
+	check(Sequencer && TargetSection && MixerTrack && Decoration);
+
+	// Initialize from existing settings
+	const FMovieSceneBoneMatchData Existing = Decoration->GetBoneMatchData();
+	BoneName = Existing.BoneName;
+	if (BoneName.IsNone())
+	{
+		// Default to the root bone (index 0)
+		if (USkeletalMeshComponent* SkelComp = FindBoundSkeletalMeshComponent(Sequencer, ObjectBinding))
+		{
+			const FReferenceSkeleton& RefSkel = SkelComp->GetSkeletalMeshAsset()->GetRefSkeleton();
+			if (RefSkel.GetNum() > 0)
+			{
+				BoneName = RefSkel.GetBoneName(0);
+			}
+		}
+	}
+	ReferenceSection = Existing.ReferenceSection;
+	MatchTimeMode = (uint8)Existing.MatchTimeMode;
+	bMatchLocationX = Existing.bMatchLocationX;
+	bMatchLocationY = Existing.bMatchLocationY;
+	bMatchLocationZ = Existing.bMatchLocationZ;
+	bMatchRotationX = Existing.bMatchRotationX;
+	bMatchRotationY = Existing.bMatchRotationY;
+	bMatchRotationZ = Existing.bMatchRotationZ;
+
+	BuildReferenceSectionList();
+	RebuildMatchTimeModeOptions();
+
+	auto MakeCheckbox = [](const FText& Label, bool& BoolRef) -> TSharedRef<SHorizontalBox>
+	{
+		return SNew(SHorizontalBox)
+			+ SHorizontalBox::Slot()
+			.AutoWidth()
+			.VAlign(VAlign_Center)
+			[
+				SNew(SCheckBox)
+				.IsChecked_Lambda([&BoolRef]() { return BoolRef ? ECheckBoxState::Checked : ECheckBoxState::Unchecked; })
+				.OnCheckStateChanged_Lambda([&BoolRef](ECheckBoxState State) { BoolRef = (State == ECheckBoxState::Checked); })
+			]
+			+ SHorizontalBox::Slot()
+			.AutoWidth()
+			.VAlign(VAlign_Center)
+			.Padding(4.f, 0.f, 12.f, 0.f)
+			[
+				SNew(STextBlock).Text(Label)
+			];
+	};
+
+	FString TargetSectionName = GetSectionDisplayName(TargetSection, Sequencer);
+
+	ChildSlot
+	[
+		SNew(SBorder)
+		.Padding(16.f)
+		[
+			SNew(SVerticalBox)
+
+			// Header showing which section we're matching
+			+ SVerticalBox::Slot()
+			.AutoHeight()
+			.Padding(0.f, 0.f, 0.f, 12.f)
+			[
+				SNew(STextBlock)
+				.Text(FText::Format(LOCTEXT("ThisSectionHeader", "This Section: {0}"), FText::FromString(TargetSectionName)))
+				.Font(FCoreStyle::GetDefaultFontStyle("Bold", 10))
+			]
+
+			// Bone picker
+			+ SVerticalBox::Slot()
+			.AutoHeight()
+			.Padding(0.f, 0.f, 0.f, 8.f)
+			[
+				SNew(SHorizontalBox)
+				+ SHorizontalBox::Slot()
+				.AutoWidth()
+				.VAlign(VAlign_Center)
+				.Padding(0.f, 0.f, 8.f, 0.f)
+				[
+					SNew(STextBlock).Text(LOCTEXT("BoneLabel", "Bone"))
+				]
+				+ SHorizontalBox::Slot()
+				.FillWidth(1.f)
+				[
+					SNew(SBoneSelectionWidget)
+					.OnGetSelectedBone_Lambda([this](bool& bMultipleValues) -> FName
+					{
+						bMultipleValues = false;
+						return BoneName;
+					})
+					.OnBoneSelectionChanged(FOnBoneSelectionChanged::CreateSP(this, &SBoneMatchDialog::OnBoneSelectionChanged))
+					.OnGetReferenceSkeleton_Lambda([this]() -> const FReferenceSkeleton&
+					{
+						static FReferenceSkeleton EmptySkeleton;
+						if (USkeletalMeshComponent* SkelComp = FindBoundSkeletalMeshComponent(Sequencer, ObjectBinding))
+						{
+							return SkelComp->GetSkeletalMeshAsset()->GetRefSkeleton();
+						}
+						return EmptySkeleton;
+					})
+				]
+			]
+
+			// Reference section
+			+ SVerticalBox::Slot()
+			.AutoHeight()
+			.Padding(0.f, 0.f, 0.f, 8.f)
+			[
+				SNew(SHorizontalBox)
+				+ SHorizontalBox::Slot()
+				.AutoWidth()
+				.VAlign(VAlign_Center)
+				.Padding(0.f, 0.f, 8.f, 0.f)
+				[
+					SNew(STextBlock).Text(LOCTEXT("ReferenceSectionLabel", "Reference Section"))
+				]
+				+ SHorizontalBox::Slot()
+				.FillWidth(1.f)
+				[
+					SNew(SComboBox<TSharedPtr<FString>>)
+					.OptionsSource(&ReferenceSectionNames)
+					.InitiallySelectedItem(ReferenceSectionNames.IsValidIndex(SelectedReferenceSectionIndex) ? ReferenceSectionNames[SelectedReferenceSectionIndex] : TSharedPtr<FString>())
+					.OnSelectionChanged(this, &SBoneMatchDialog::OnReferenceSectionChanged)
+					.OnGenerateWidget_Lambda([](TSharedPtr<FString> Item)
+					{
+						return SNew(STextBlock).Text(FText::FromString(*Item));
+					})
+					.Content()
+					[
+						SNew(STextBlock)
+						.Text_Lambda([this]()
+						{
+							if (SelectedReferenceSectionIndex >= 0 && SelectedReferenceSectionIndex < ReferenceSectionNames.Num())
+							{
+								return FText::FromString(*ReferenceSectionNames[SelectedReferenceSectionIndex]);
+							}
+							return LOCTEXT("NoRefSection", "(None)");
+						})
+					]
+				]
+			]
+
+			// Match time mode
+			+ SVerticalBox::Slot()
+			.AutoHeight()
+			.Padding(0.f, 0.f, 0.f, 8.f)
+			[
+				SNew(SHorizontalBox)
+				+ SHorizontalBox::Slot()
+				.AutoWidth()
+				.VAlign(VAlign_Center)
+				.Padding(0.f, 0.f, 8.f, 0.f)
+				[
+					SNew(STextBlock).Text(LOCTEXT("MatchTimeLabel", "Match Time"))
+				]
+				+ SHorizontalBox::Slot()
+				.FillWidth(1.f)
+				[
+					SNew(SComboBox<TSharedPtr<FString>>)
+					.OptionsSource(&MatchTimeModeNames)
+					.InitiallySelectedItem(MatchTimeModeNames.IsValidIndex(SelectedMatchTimeModeIndex) ? MatchTimeModeNames[SelectedMatchTimeModeIndex] : TSharedPtr<FString>())
+					.OnSelectionChanged(this, &SBoneMatchDialog::OnMatchTimeModeChanged)
+					.OnGenerateWidget_Lambda([](TSharedPtr<FString> Item)
+					{
+						return SNew(STextBlock).Text(FText::FromString(*Item));
+					})
+					.Content()
+					[
+						SNew(STextBlock)
+						.Text_Lambda([this]()
+						{
+							if (SelectedMatchTimeModeIndex >= 0 && SelectedMatchTimeModeIndex < MatchTimeModeNames.Num())
+							{
+								return FText::FromString(*MatchTimeModeNames[SelectedMatchTimeModeIndex]);
+							}
+							return LOCTEXT("NoMatchTime", "Start of Section");
+						})
+					]
+				]
+			]
+
+			// Location matching
+			+ SVerticalBox::Slot()
+			.AutoHeight()
+			.Padding(0.f, 4.f, 0.f, 4.f)
+			[
+				SNew(SHorizontalBox)
+				+ SHorizontalBox::Slot()
+				.AutoWidth()
+				.VAlign(VAlign_Center)
+				.Padding(0.f, 0.f, 8.f, 0.f)
+				[
+					SNew(STextBlock).Text(LOCTEXT("LocationLabel", "Match Location"))
+				]
+				+ SHorizontalBox::Slot().AutoWidth() [ MakeCheckbox(LOCTEXT("X", "X"), bMatchLocationX) ]
+				+ SHorizontalBox::Slot().AutoWidth() [ MakeCheckbox(LOCTEXT("Y", "Y"), bMatchLocationY) ]
+				+ SHorizontalBox::Slot().AutoWidth() [ MakeCheckbox(LOCTEXT("Z", "Z"), bMatchLocationZ) ]
+			]
+
+			// Rotation matching
+			+ SVerticalBox::Slot()
+			.AutoHeight()
+			.Padding(0.f, 0.f, 0.f, 8.f)
+			[
+				SNew(SHorizontalBox)
+				+ SHorizontalBox::Slot()
+				.AutoWidth()
+				.VAlign(VAlign_Center)
+				.Padding(0.f, 0.f, 8.f, 0.f)
+				[
+					SNew(STextBlock).Text(LOCTEXT("RotationLabel", "Match Rotation"))
+				]
+				+ SHorizontalBox::Slot().AutoWidth() [ MakeCheckbox(LOCTEXT("Roll", "Roll"), bMatchRotationX) ]
+				+ SHorizontalBox::Slot().AutoWidth() [ MakeCheckbox(LOCTEXT("Pitch", "Pitch"), bMatchRotationY) ]
+				+ SHorizontalBox::Slot().AutoWidth() [ MakeCheckbox(LOCTEXT("Yaw", "Yaw"), bMatchRotationZ) ]
+			]
+
+			// OK / Cancel
+			+ SVerticalBox::Slot()
+			.AutoHeight()
+			.Padding(0.f, 8.f, 0.f, 0.f)
+			[
+				SNew(SHorizontalBox)
+				+ SHorizontalBox::Slot().FillWidth(1.f) [ SNew(SSpacer) ]
+				+ SHorizontalBox::Slot()
+				.AutoWidth()
+				.Padding(8.f, 0.f)
+				[
+					SNew(SButton)
+					.HAlign(HAlign_Center)
+					.ContentPadding(FAppStyle::GetMargin("StandardDialog.ContentPadding"))
+					.Text(LOCTEXT("OK", "OK"))
+					.OnClicked(this, &SBoneMatchDialog::OnOKClicked)
+					.IsEnabled_Lambda([this]() { return !BoneName.IsNone(); })
+				]
+				+ SHorizontalBox::Slot()
+				.AutoWidth()
+				[
+					SNew(SButton)
+					.HAlign(HAlign_Center)
+					.ContentPadding(FAppStyle::GetMargin("StandardDialog.ContentPadding"))
+					.Text(LOCTEXT("Cancel", "Cancel"))
+					.OnClicked_Lambda([this]()
+					{
+						CloseDialog();
+						return FReply::Handled();
+					})
+				]
+			]
+		]
+	];
+}
+
+void SBoneMatchDialog::BuildReferenceSectionList()
+{
+	ReferenceSections.Reset();
+	ReferenceSectionNames.Reset();
+	SelectedReferenceSectionIndex = INDEX_NONE;
+
+	if (!TargetSection || !MixerTrack)
+	{
+		return;
+	}
+
+	for (UMovieSceneSection* Section : MixerTrack->GetAllSections())
+	{
+		if (!Section || Section == TargetSection)
+		{
+			continue;
+		}
+
+		if (Section->IsA<UMovieSceneAnimTransitionSectionBase>())
+		{
+			continue;
+		}
+
+		// Only include sections that overlap in time
+		if (TRange<FFrameNumber>::Intersection(TargetSection->GetRange(), Section->GetRange()).IsEmpty())
+		{
+			continue;
+		}
+
+		// Exclude sections whose bone match chain references the target,
+		// since that would create a dependency cycle
+		if (WouldCreateBoneMatchCycle(Section, TargetSection))
+		{
+			continue;
+		}
+
+		ReferenceSections.Add(Section);
+
+		FString Name = GetSectionDisplayName(Section, Sequencer);
+		if (Section->HasStartFrame() && Section->HasEndFrame() && Sequencer)
+		{
+			FFrameRate TickResolution = Sequencer->GetFocusedTickResolution();
+			FFrameRate DisplayRate = Sequencer->GetFocusedDisplayRate();
+			FFrameNumber StartDisplay = FFrameRate::TransformTime(FFrameTime(Section->GetInclusiveStartFrame()), TickResolution, DisplayRate).FloorToFrame();
+			FFrameNumber EndDisplay = FFrameRate::TransformTime(FFrameTime(Section->GetExclusiveEndFrame()), TickResolution, DisplayRate).FloorToFrame();
+			Name += FString::Printf(TEXT(" [%d-%d]"), StartDisplay.Value, EndDisplay.Value);
+		}
+		ReferenceSectionNames.Add(MakeShared<FString>(Name));
+
+		// Select the one that matches current settings
+		if (ReferenceSection.IsValid() && ReferenceSection.Get() == Section)
+		{
+			SelectedReferenceSectionIndex = ReferenceSections.Num() - 1;
+		}
+	}
+
+	// Default to the first reference section if none was previously selected
+	if (SelectedReferenceSectionIndex == INDEX_NONE && ReferenceSections.Num() > 0)
+	{
+		SelectedReferenceSectionIndex = 0;
+		ReferenceSection = ReferenceSections[0];
+	}
+}
+
+void SBoneMatchDialog::RebuildMatchTimeModeOptions()
+{
+	FilteredMatchTimeModes.Reset();
+	MatchTimeModeNames.Reset();
+	SelectedMatchTimeModeIndex = 0;
+
+	UMovieSceneSection* RefSection = nullptr;
+	FString RefName = TEXT("Reference");
+	if (SelectedReferenceSectionIndex >= 0 && SelectedReferenceSectionIndex < ReferenceSections.Num())
+	{
+		RefSection = ReferenceSections[SelectedReferenceSectionIndex].Get();
+		if (RefSection)
+		{
+			RefName = GetSectionDisplayName(RefSection, Sequencer);
+		}
+	}
+
+	// Compute the overlap range between target and reference sections.
+	// Only match times that fall within the overlap are valid.
+	TRange<FFrameNumber> Overlap = TRange<FFrameNumber>::All();
+	if (TargetSection && RefSection)
+	{
+		TRange<FFrameNumber> TargetRange = TargetSection->GetRange();
+		TRange<FFrameNumber> RefRange = RefSection->GetRange();
+		Overlap = TRange<FFrameNumber>::Intersection(TargetRange, RefRange);
+	}
+
+	auto AddMode = [&](EBoneMatchTimeMode Mode, const FString& Label)
+	{
+		FilteredMatchTimeModes.Add(Mode);
+		MatchTimeModeNames.Add(MakeShared<FString>(Label));
+
+		if ((uint8)Mode == MatchTimeMode)
+		{
+			SelectedMatchTimeModeIndex = FilteredMatchTimeModes.Num() - 1;
+		}
+	};
+
+	// For each match time, check that it falls within the other section's range
+	// so both sections are active at the match point.
+	TRange<FFrameNumber> TargetRange = TargetSection ? TargetSection->GetRange() : TRange<FFrameNumber>::Empty();
+	TRange<FFrameNumber> RefRange = RefSection ? RefSection->GetRange() : TRange<FFrameNumber>::Empty();
+
+	if (Sequencer)
+	{
+		FFrameNumber CurrentFrame = Sequencer->GetLocalTime().Time.FloorToFrame();
+		if (TargetRange.Contains(CurrentFrame) && RefRange.Contains(CurrentFrame))
+		{
+			AddMode(EBoneMatchTimeMode::AtCurrentTime, TEXT("At Current Time"));
+		}
+	}
+
+	if (TargetSection && TargetSection->HasStartFrame() && RefRange.Contains(TargetSection->GetInclusiveStartFrame()))
+	{
+		AddMode(EBoneMatchTimeMode::AtStartOfSelectedSection, TEXT("At Start of This Section"));
+	}
+	if (TargetSection && TargetSection->HasEndFrame() && RefRange.Contains(TargetSection->GetExclusiveEndFrame() - 1))
+	{
+		AddMode(EBoneMatchTimeMode::AtEndOfSelectedSection, TEXT("At End of This Section"));
+	}
+
+	if (RefSection && RefSection->HasStartFrame() && TargetRange.Contains(RefSection->GetInclusiveStartFrame()))
+	{
+		AddMode(EBoneMatchTimeMode::AtStartOfReferenceSection,
+			FString::Printf(TEXT("At Start of %s"), *RefName));
+	}
+	if (RefSection && RefSection->HasEndFrame() && TargetRange.Contains(RefSection->GetExclusiveEndFrame() - 1))
+	{
+		AddMode(EBoneMatchTimeMode::AtEndOfReferenceSection,
+			FString::Printf(TEXT("At End of %s"), *RefName));
+	}
+
+	if (!Overlap.IsEmpty())
+	{
+		AddMode(EBoneMatchTimeMode::InBetween, TEXT("In Between (Midpoint)"));
+	}
+
+	// Sync the backing value with the displayed selection. When the previously
+	// stored mode isn't in the filtered list, the UI defaults to index 0 but
+	// MatchTimeMode still holds the old (now-invalid) enum value.
+	if (SelectedMatchTimeModeIndex >= 0 && SelectedMatchTimeModeIndex < FilteredMatchTimeModes.Num())
+	{
+		MatchTimeMode = (uint8)FilteredMatchTimeModes[SelectedMatchTimeModeIndex];
+	}
+}
+
+void SBoneMatchDialog::OnBoneSelectionChanged(FName NewBoneName)
+{
+	BoneName = NewBoneName;
+}
+
+void SBoneMatchDialog::OnReferenceSectionChanged(TSharedPtr<FString> NewValue, ESelectInfo::Type SelectInfo)
+{
+	int32 Index = ReferenceSectionNames.IndexOfByKey(NewValue);
+	if (Index != INDEX_NONE)
+	{
+		SelectedReferenceSectionIndex = Index;
+		ReferenceSection = ReferenceSections[Index];
+		RebuildMatchTimeModeOptions();
+	}
+}
+
+void SBoneMatchDialog::OnMatchTimeModeChanged(TSharedPtr<FString> NewValue, ESelectInfo::Type SelectInfo)
+{
+	int32 Index = MatchTimeModeNames.IndexOfByKey(NewValue);
+	if (Index != INDEX_NONE && Index < FilteredMatchTimeModes.Num())
+	{
+		SelectedMatchTimeModeIndex = Index;
+		MatchTimeMode = (uint8)FilteredMatchTimeModes[Index];
+	}
+}
+
+FReply SBoneMatchDialog::OnOKClicked()
+{
+	FMovieSceneBoneMatchData MatchData;
+	MatchData.BoneName = BoneName;
+	MatchData.ReferenceSection = ReferenceSection;
+	MatchData.MatchTimeMode = (EBoneMatchTimeMode)MatchTimeMode;
+	MatchData.bMatchLocationX = bMatchLocationX;
+	MatchData.bMatchLocationY = bMatchLocationY;
+	MatchData.bMatchLocationZ = bMatchLocationZ;
+	MatchData.bMatchRotationX = bMatchRotationX;
+	MatchData.bMatchRotationY = bMatchRotationY;
+	MatchData.bMatchRotationZ = bMatchRotationZ;
+
+	OnAccepted.ExecuteIfBound(MatchData);
+	CloseDialog();
+	return FReply::Handled();
+}
+
+FReply SBoneMatchDialog::OpenDialog(bool bModal)
+{
+	check(!DialogWindow.IsValid());
+
+	TSharedRef<SWindow> Window = SNew(SWindow)
+		.Title(LOCTEXT("BoneMatchDialogTitle", "Bone Match Settings"))
+		.CreateTitleBar(true)
+		.Type(EWindowType::Normal)
+		.SizingRule(ESizingRule::Autosized)
+		.AutoCenter(EAutoCenter::PreferredWorkArea)
+		.FocusWhenFirstShown(true)
+		.ActivationPolicy(EWindowActivationPolicy::FirstShown)
+		[
+			AsShared()
+		];
+
+	Window->SetWidgetToFocusOnActivate(AsShared());
+	DialogWindow = Window;
+
+	if (bModal)
+	{
+		GEditor->EditorAddModalWindow(Window);
+	}
+	else
+	{
+		FSlateApplication::Get().AddWindow(Window);
+	}
+
+	return FReply::Handled();
+}
+
+void SBoneMatchDialog::CloseDialog()
+{
+	if (DialogWindow.IsValid())
+	{
+		DialogWindow.Pin()->RequestDestroyWindow();
+		DialogWindow.Reset();
+	}
+}
+
+#undef LOCTEXT_NAMESPACE
diff --git a/Engine/Plugins/Experimental/MovieSceneAnimMixer/Source/MovieSceneAnimMixerEditor/Private/SBoneMatchDialog.h b/Engine/Plugins/Experimental/MovieSceneAnimMixer/Source/MovieSceneAnimMixerEditor/Private/SBoneMatchDialog.h
new file mode 100644
--- /dev/null
+++ b/Engine/Plugins/Experimental/MovieSceneAnimMixer/Source/MovieSceneAnimMixerEditor/Private/SBoneMatchDialog.h
@@ -0,0 +1,86 @@
+// Copyright Epic Games, Inc. All Rights Reserved.
+
+#pragma once
+
+#include "Misc/Guid.h"
+#include "MovieSceneAnimationMixerTrack.h"
+#include "Widgets/DeclarativeSyntaxSupport.h"
+#include "Widgets/SCompoundWidget.h"
+
+class ISequencer;
+class UMovieSceneSection;
+class UMovieSceneAnimationMixerTrack;
+class UMovieSceneRootMotionSettingsDecoration;
+class SWindow;
+struct FMovieSceneBoneMatchData;
+
+DECLARE_DELEGATE_OneParam(FOnBoneMatchAccepted, const FMovieSceneBoneMatchData&);
+
+// Modal dialog for configuring bone matching parameters.
+class SBoneMatchDialog : public SCompoundWidget
+{
+public:
+
+	SLATE_BEGIN_ARGS(SBoneMatchDialog)
+		: _Sequencer(nullptr)
+		, _TargetSection(nullptr)
+		, _MixerTrack(nullptr)
+		, _Decoration(nullptr)
+		{}
+		SLATE_ARGUMENT(ISequencer*, Sequencer)
+		SLATE_ARGUMENT(UMovieSceneSection*, TargetSection)
+		SLATE_ARGUMENT(UMovieSceneAnimationMixerTrack*, MixerTrack)
+		SLATE_ARGUMENT(UMovieSceneRootMotionSettingsDecoration*, Decoration)
+		SLATE_ARGUMENT(FGuid, ObjectBinding)
+		SLATE_EVENT(FOnBoneMatchAccepted, OnAccepted)
+	SLATE_END_ARGS()
+
+	void Construct(const FArguments& InArgs);
+
+	FReply OpenDialog(bool bModal = true);
+	void CloseDialog();
+
+private:
+
+	// Build the list of candidate reference sections
+	void BuildReferenceSectionList();
+
+	// Build the match time mode option list
+	void RebuildMatchTimeModeOptions();
+
+	// Callbacks
+	void OnBoneSelectionChanged(FName NewBoneName);
+	void OnReferenceSectionChanged(TSharedPtr<FString> NewValue, ESelectInfo::Type SelectInfo);
+	void OnMatchTimeModeChanged(TSharedPtr<FString> NewValue, ESelectInfo::Type SelectInfo);
+	FReply OnOKClicked();
+
+	ISequencer* Sequencer = nullptr;
+	UMovieSceneSection* TargetSection = nullptr;
+	UMovieSceneAnimationMixerTrack* MixerTrack = nullptr;
+	UMovieSceneRootMotionSettingsDecoration* Decoration = nullptr;
+	FGuid ObjectBinding;
+	FOnBoneMatchAccepted OnAccepted;
+
+	TWeakPtr<SWindow> DialogWindow;
+
+	FName BoneName;
+	TWeakObjectPtr<UMovieSceneSection> ReferenceSection;
+	uint8 MatchTimeMode = 0;
+
+	bool bMatchLocationX = true;
+	bool bMatchLocationY = true;
+	bool bMatchLocationZ = false;
+	bool bMatchRotationX = false;
+	bool bMatchRotationY = false;
+	bool bMatchRotationZ = true;
+
+	// Reference section options
+	TArray<TWeakObjectPtr<UMovieSceneSection>> ReferenceSections;
+	TArray<TSharedPtr<FString>> ReferenceSectionNames;
+	int32 SelectedReferenceSectionIndex = INDEX_NONE;
+
+	// Match time mode options (filtered to only include valid modes for current overlap)
+	TArray<EBoneMatchTimeMode> FilteredMatchTimeModes;
+	TArray<TSharedPtr<FString>> MatchTimeModeNames;
+	int32 SelectedMatchTimeModeIndex = 0;
+};
diff --git a/Engine/Source/Editor/MovieSceneTools/Private/TrackEditors/CommonAnimationTrackEditor.cpp b/Engine/Source/Editor/MovieSceneTools/Private/TrackEditors/CommonAnimationTrackEditor.cpp
--- a/Engine/Source/Editor/MovieSceneTools/Private/TrackEditors/CommonAnimationTrackEditor.cpp
+++ b/Engine/Source/Editor/MovieSceneTools/Private/TrackEditors/CommonAnimationTrackEditor.cpp
@@ -1023,6 +1023,11 @@ void FCommonAnimationSection::BuildSectionContextMenu(FMenuBuilder& MenuBuilder,
 			);
 	};
 
+	if (!ShouldShowLegacyBoneMatching())
+	{
+		return;
+	}
+
 	MenuBuilder.BeginSection(NAME_None, LOCTEXT("MotionBlendingOptions", "Motion Blending Options"));
 	{
 		MenuBuilder.AddSubMenu(
diff --git a/Engine/Source/Editor/MovieSceneTools/Public/TrackEditors/CommonAnimationTrackEditor.h b/Engine/Source/Editor/MovieSceneTools/Public/TrackEditors/CommonAnimationTrackEditor.h
--- a/Engine/Source/Editor/MovieSceneTools/Public/TrackEditors/CommonAnimationTrackEditor.h
+++ b/Engine/Source/Editor/MovieSceneTools/Public/TrackEditors/CommonAnimationTrackEditor.h
@@ -177,6 +177,9 @@ public:
 	UE_API virtual bool RequestDeleteKeyArea(const TArray<FName>& KeyAreaNamePath) override;
 	UE_API virtual void BuildSectionContextMenu(FMenuBuilder& MenuBuilder, const FGuid& InObjectBinding) override;
 
+	// Whether to show the legacy bone matching UI from MovieSceneCommonAnimationTrack.
+	// Override to return false when using a different bone matching system (e.g. AnimMixer).
+	virtual bool ShouldShowLegacyBoneMatching() const { return true; }
 
 protected:
 	UE_API void FindBestBlendSection(FGuid InObjectBinding);
diff --git a/Engine/Source/Editor/Sequencer/Private/IKeyArea.cpp b/Engine/Source/Editor/Sequencer/Private/IKeyArea.cpp
--- a/Engine/Source/Editor/Sequencer/Private/IKeyArea.cpp
+++ b/Engine/Source/Editor/Sequencer/Private/IKeyArea.cpp
@@ -1,6 +1,7 @@
 // Copyright Epic Games, Inc. All Rights Reserved.
 
 #include "IKeyArea.h"
+#include "ISequencerDecorationEditor.h"
 #include "ISequencerModule.h"
 #include "Modules/ModuleManager.h"
 #include "Widgets/SNullWidget.h"
@@ -298,9 +299,22 @@ TSharedRef<SWidget> IKeyArea::CreateKeyEditor(TWeakPtr<ISequencer> Sequencer, co
 	const FMovieSceneChannelMetaData* MetaData = ChannelHandle.GetMetaData();
 	if (EditorInterface && OwningSection && MetaData)
 	{
+		// Let decoration editors customize key editors for their channels
+		UObject* OwningObject = MetaData->WeakOwningObject.Get();
+		if (OwningObject && OwningObject != OwningSection)
+		{
+			if (ISequencerDecorationEditor* DecorationEditor = ISequencerDecorationEditor::FindDecorationEditor(OwningObject->GetClass()))
+			{
+				if (TSharedPtr<SWidget> CustomEditor = DecorationEditor->CreateKeyEditor(*OwningObject, *OwningSection, ChannelHandle, Sequencer, ObjectBindingID))
+				{
+					return CustomEditor.ToSharedRef();
+				}
+			}
+		}
+
 		FCreateKeyEditorParams Params{
 			OwningSection,
-			MetaData->WeakOwningObject.Get(),
+			OwningObject,
 			Sequencer.Pin().ToSharedRef(),
 			ObjectBindingID,
 			PropertyBindingsPtr
diff --git a/Engine/Source/Editor/Sequencer/Private/MVVM/SectionModelStorageExtension.cpp b/Engine/Source/Editor/Sequencer/Private/MVVM/SectionModelStorageExtension.cpp
--- a/Engine/Source/Editor/Sequencer/Private/MVVM/SectionModelStorageExtension.cpp
+++ b/Engine/Source/Editor/Sequencer/Private/MVVM/SectionModelStorageExtension.cpp
@@ -2,7 +2,10 @@
 
 #include "MVVM/SectionModelStorageExtension.h"
 
+#include "ISequencer.h"
+#include "ISequencerDecorationEditor.h"
 #include "MVVM/ViewModels/SectionModel.h"
+#include "MVVM/ViewModels/SequenceModel.h"
 #include "MovieSceneSection.h"
 #include "ISequencerCoreModule.h"
 #include "Modules/ModuleManager.h"
@@ -25,6 +28,11 @@ FSectionModelStorageExtension::FSectionModelStorageExtension()
 {
 }
 
+void FSectionModelStorageExtension::OnCreated(TSharedRef<FViewModel> InOwner)
+{
+	WeakOwner = InOwner;
+}
+
 void FSectionModelStorageExtension::OnReinitialize()
 {
 	for (auto It = SectionToModel.CreateIterator(); It; ++It)
@@ -63,6 +71,8 @@ TSharedPtr<FSectionModel> FSectionModelStorageExtension::CreateModelForSection(U
 
 	SectionToModel.Add(SectionAsKey, SectionModel);
 
+	NotifyDecorationEditors(InSection);
+
 	return SectionModel;
 }
 
@@ -72,5 +82,38 @@ TSharedPtr<FSectionModel> FSectionModelStorageExtension::FindModelForSection(con
 	return SectionToModel.FindRef(SectionAsKey).Pin();
 }
 
+void FSectionModelStorageExtension::NotifyDecorationEditors(UMovieSceneSection* InSection)
+{
+	if (!InSection)
+	{
+		return;
+	}
+
+	TSharedPtr<FViewModel> Owner = WeakOwner.Pin();
+	if (!Owner)
+	{
+		return;
+	}
+
+	FSequenceModel* SequenceModel = Owner->CastThis<FSequenceModel>();
+	if (!SequenceModel)
+	{
+		return;
+	}
+
+	TWeakPtr<ISequencer> Sequencer = SequenceModel->GetSequencer();
+
+	for (TObjectPtr<UObject> Decoration : InSection->GetDecorations())
+	{
+		if (Decoration)
+		{
+			if (ISequencerDecorationEditor* Editor = ISequencerDecorationEditor::FindDecorationEditor(Decoration->GetClass()))
+			{
+				Editor->OnSectionInterfaceCreated(*InSection, *Decoration, Sequencer);
+			}
+		}
+	}
+}
+
 } // namespace UE::Sequencer
 
diff --git a/Engine/Source/Editor/Sequencer/Private/MVVM/ViewModels/ChannelModel.cpp b/Engine/Source/Editor/Sequencer/Private/MVVM/ViewModels/ChannelModel.cpp
--- a/Engine/Source/Editor/Sequencer/Private/MVVM/ViewModels/ChannelModel.cpp
+++ b/Engine/Source/Editor/Sequencer/Private/MVVM/ViewModels/ChannelModel.cpp
@@ -35,6 +35,7 @@
 #include "MovieSceneTrack.h"
 #include "MVVM/Views/SChannelView.h"
 #include "PropertyEditorModule.h"
+#include "ISequencerDecorationEditor.h"
 #include "SKeyAreaEditorSwitcher.h"
 #include "SequencerCoreFwd.h"
 #include "SequencerSectionPainter.h"
@@ -1037,8 +1038,52 @@ TSharedPtr<SWidget> FChannelGroupOutlinerModel::CreateOutlinerViewForColumn(cons
 		return false;
 	};
 
+	// Check if any channel's decoration editor provides a custom keyframe widget
+	// (e.g., bone match channel replaces "add key" with a dialog launcher).
+	// Used by both KeyFrame and Nav columns since the add-key button can appear in either.
+	auto TryCreateDecorationKeyFrameWidget = [this, &Editor]() -> TSharedPtr<SWidget>
+	{
+		for (const TWeakViewModelPtr<FChannelModel>& WeakChannel : GetChannels())
+		{
+			if (TSharedPtr<FChannelModel> Channel = WeakChannel.Pin())
+			{
+				if (TSharedPtr<IKeyArea> KeyArea = Channel->GetKeyArea())
+				{
+					const FMovieSceneChannelMetaData* MetaData = KeyArea->GetChannel().GetMetaData();
+					UObject* OwningObject = MetaData ? MetaData->WeakOwningObject.Get() : nullptr;
+					UMovieSceneSection* Section = KeyArea->GetOwningSection();
+
+					if (OwningObject && Section && OwningObject != Section)
+					{
+						if (ISequencerDecorationEditor* DecorationEditor = ISequencerDecorationEditor::FindDecorationEditor(OwningObject->GetClass()))
+						{
+							FGuid ObjectBindingID;
+							if (TViewModelPtr<IObjectBindingExtension> ObjectBinding = FindAncestorOfType<IObjectBindingExtension>())
+							{
+								ObjectBindingID = ObjectBinding->GetObjectGuid();
+							}
+
+							if (TSharedPtr<SWidget> CustomWidget = DecorationEditor->CreateKeyFrameWidget(
+								*OwningObject, *Section, KeyArea->GetChannel(), Editor->GetSequencer(), ObjectBindingID,
+								SharedThis(this)))
+							{
+								return CustomWidget;
+							}
+						}
+					}
+				}
+			}
+		}
+		return nullptr;
+	};
+
 	if (InColumnName == FCommonOutlinerNames::KeyFrame)
 	{
+		if (TSharedPtr<SWidget> CustomWidget = TryCreateDecorationKeyFrameWidget())
+		{
+			return CustomWidget;
+		}
+
 		if (!CanAnyChannelBeKeyed())
 		{
 			return nullptr;
@@ -1051,6 +1096,11 @@ TSharedPtr<SWidget> FChannelGroupOutlinerModel::CreateOutlinerViewForColumn(cons
 
 	if (InColumnName == FCommonOutlinerNames::Nav)
 	{
+		if (TSharedPtr<SWidget> CustomWidget = TryCreateDecorationKeyFrameWidget())
+		{
+			return CustomWidget;
+		}
+
 		if (!CanAnyChannelBeKeyed())
 		{
 			return nullptr;
diff --git a/Engine/Source/Editor/Sequencer/Private/MVVM/ViewModels/DecorationOutlinerModel.cpp b/Engine/Source/Editor/Sequencer/Private/MVVM/ViewModels/DecorationOutlinerModel.cpp
--- a/Engine/Source/Editor/Sequencer/Private/MVVM/ViewModels/DecorationOutlinerModel.cpp
+++ b/Engine/Source/Editor/Sequencer/Private/MVVM/ViewModels/DecorationOutlinerModel.cpp
@@ -61,6 +61,18 @@ bool FDecorationOutlinerModel::NeedsReconstruct() const
 			return true;
 		}
 	}
+
+	// Check if any channel decoration flagged its proxy as dirty
+	for (const TObjectPtr<UObject>& Decoration : CurrentDecorations)
+	{
+		if (IMovieSceneChannelDecoration* ChannelDecoration = Cast<IMovieSceneChannelDecoration>(Decoration.Get()))
+		{
+			if (ChannelDecoration->ConsumeDecorationChannelProxyDirty())
+			{
+				return true;
+			}
+		}
+	}
 	return false;
 }
 
@@ -144,8 +156,12 @@ void FDecorationOutlinerModel::OnConstruct()
 				{
 					continue;
 				}
-				const TSharedRef<ISequencerSection> SectionInterface = DecorationEditor->MakeSectionInterface(*Section, *WeakParentTrack.Get(), WeakParentTrack.Get()->FindObjectBindingGuid());
-				const TSharedPtr<FSectionModel> SectionModel = SectionModelStorage->CreateModelForSection(Section, SectionInterface);
+				const TSharedPtr<ISequencerSection> SectionInterface = DecorationEditor->MakeSectionInterface(*Section, *WeakParentTrack.Get(), WeakParentTrack.Get()->FindObjectBindingGuid());
+				if (!SectionInterface)
+				{
+					continue;
+				}
+				const TSharedPtr<FSectionModel> SectionModel = SectionModelStorage->CreateModelForSection(Section, SectionInterface.ToSharedRef());
 
 				TSharedPtr<FSectionOutlinerModel> SectionOutlinerModel;
 				// Search in RecycledChildren (where we moved existing models) rather than OutlinerChildren (which is now empty)
diff --git a/Engine/Source/Editor/Sequencer/Private/Sequencer.cpp b/Engine/Source/Editor/Sequencer/Private/Sequencer.cpp
--- a/Engine/Source/Editor/Sequencer/Private/Sequencer.cpp
+++ b/Engine/Source/Editor/Sequencer/Private/Sequencer.cpp
@@ -2524,6 +2524,15 @@ void FSequencer::TransformSelection(FFrameTime InDeltaTime, float InScale)
 		Pair.Key->SetRange(Pair.Value);
 	}
 
+	// Notify tracks that sections moved so they can update dependent state
+	for (TTuple<UMovieSceneSection*, TRange<FFrameNumber>>& Pair : SectionToNewBounds)
+	{
+		if (UMovieSceneTrack* Track = Pair.Key->GetTypedOuter<UMovieSceneTrack>())
+		{
+			Track->OnSectionMoved(*Pair.Key, FMovieSceneSectionMovedParams(EPropertyChangeType::ValueSet));
+		}
+	}
+
 	if (bAnythingChanged)
 	{
 		NotifyMovieSceneDataChanged( EMovieSceneDataChangeType::TrackValueChanged );
diff --git a/Engine/Source/Editor/Sequencer/Private/SequencerHotspots.cpp b/Engine/Source/Editor/Sequencer/Private/SequencerHotspots.cpp
--- a/Engine/Source/Editor/Sequencer/Private/SequencerHotspots.cpp
+++ b/Engine/Source/Editor/Sequencer/Private/SequencerHotspots.cpp
@@ -22,13 +22,42 @@
 #include "MovieSceneTimeHelpers.h"
 #include "SequencerCommonHelpers.h"
 #include "ISequencerSection.h"
+#include "ISequencerDecorationEditor.h"
 #include "MovieSceneSignedObject.h"
+#include "Decorations/MovieSceneDecorationContainer.h"
 
 #include "MVVM/ViewModels/SectionModel.h"
 #include "MVVM/ViewModels/VirtualTrackArea.h"
 
 #define LOCTEXT_NAMESPACE "SequencerHotspots"
 
+// Iterate section decorations and let registered decoration editors
+// contribute to the context menu.
+static void BuildDecorationContextMenu(
+	FMenuBuilder& MenuBuilder,
+	UMovieSceneSection* Section,
+	const FGuid& ObjectBinding,
+	TWeakPtr<ISequencer> Sequencer)
+{
+	if (!Section)
+	{
+		return;
+	}
+
+	for (TObjectPtr<UObject> DecorationObject : Section->GetDecorations())
+	{
+		if (!DecorationObject)
+		{
+			continue;
+		}
+
+		if (ISequencerDecorationEditor* Editor = ISequencerDecorationEditor::FindDecorationEditor(DecorationObject->GetClass()))
+		{
+			Editor->BuildSectionContextMenu(MenuBuilder, *Section, *DecorationObject, ObjectBinding, Sequencer);
+		}
+	}
+}
+
 namespace UE
 {
 namespace Sequencer
@@ -311,7 +340,9 @@ bool FKeyBarHotspot::PopulateContextMenu(FMenuBuilder& MenuBuilder, TSharedPtr<F
 	FSectionContextMenu::BuildMenu(MenuBuilder, MenuExtender, WeakSequencer, MouseDownTime);
 
 	TSharedPtr<IObjectBindingExtension> ObjectBinding = SectionModel->FindAncestorOfType<IObjectBindingExtension>();
-	SectionModel->GetSectionInterface()->BuildSectionContextMenu(MenuBuilder, ObjectBinding ? ObjectBinding->GetObjectGuid() : FGuid());
+	FGuid ObjectBindingGuid = ObjectBinding ? ObjectBinding->GetObjectGuid() : FGuid();
+	SectionModel->GetSectionInterface()->BuildSectionContextMenu(MenuBuilder, ObjectBindingGuid);
+	BuildDecorationContextMenu(MenuBuilder, SectionModel->GetSection(), ObjectBindingGuid, Sequencer);
 
 	return true;
 }
@@ -563,7 +594,9 @@ bool FSectionHotspotBase::PopulateContextMenu(FMenuBuilder& MenuBuilder, TShared
 		FSectionContextMenu::BuildMenu(MenuBuilder, MenuExtender, WeakSequencer, MouseDownTime);
 
 		TSharedPtr<IObjectBindingExtension> ObjectBinding = SectionModel->FindAncestorOfType<IObjectBindingExtension>();
-		SectionModel->GetSectionInterface()->BuildSectionContextMenu(MenuBuilder, ObjectBinding ? ObjectBinding->GetObjectGuid() : FGuid());
+		FGuid ObjectBindingGuid = ObjectBinding ? ObjectBinding->GetObjectGuid() : FGuid();
+		SectionModel->GetSectionInterface()->BuildSectionContextMenu(MenuBuilder, ObjectBindingGuid);
+		BuildDecorationContextMenu(MenuBuilder, ThisSection, ObjectBindingGuid, WeakSequencer.Pin());
 	}
 
 	return true;
@@ -678,7 +711,9 @@ bool FSectionEasingAreaHotspot::PopulateContextMenu(FMenuBuilder& MenuBuilder, T
 	if (ThisSection)
 	{
 		TSharedPtr<IObjectBindingExtension> ObjectBinding = SectionModel->FindAncestorOfType<IObjectBindingExtension>();
-		SectionModel->GetSectionInterface()->BuildSectionContextMenu(MenuBuilder, ObjectBinding ? ObjectBinding->GetObjectGuid() : FGuid());
+		FGuid ObjectBindingGuid = ObjectBinding ? ObjectBinding->GetObjectGuid() : FGuid();
+		SectionModel->GetSectionInterface()->BuildSectionContextMenu(MenuBuilder, ObjectBindingGuid);
+		BuildDecorationContextMenu(MenuBuilder, ThisSection, ObjectBindingGuid, Sequencer);
 	}
 
 	return true;
diff --git a/Engine/Source/Editor/Sequencer/Public/ISequencerDecorationEditor.h b/Engine/Source/Editor/Sequencer/Public/ISequencerDecorationEditor.h
--- a/Engine/Source/Editor/Sequencer/Public/ISequencerDecorationEditor.h
+++ b/Engine/Source/Editor/Sequencer/Public/ISequencerDecorationEditor.h
@@ -3,27 +3,31 @@
 #pragma once
 
 #include "CoreMinimal.h"
+#include "Channels/MovieSceneChannelHandle.h"
 #include "Features/IModularFeatures.h"
 
+class FMenuBuilder;
+class ISequencer;
 class ISequencerSection;
+class SWidget;
 class UMovieSceneSection;
 class UMovieSceneTrack;
 
+namespace UE::Sequencer { class FViewModel; }
+
 /**
- * Interface for decoration editors that can create section interfaces for decorations
- * that provide their own sections (via IMovieSceneSectionProviderDecoration).
+ * Editor-side counterpart for section decorations, registered via modular features.
  *
- * Decorations are generic and can be added to any track type, so the parent track's
- * track editor cannot know how to create section interfaces for decoration-provided sections.
- * Instead, each decoration that provides sections should have a corresponding decoration editor
- * registered via modular features.
+ * Decoration editors can create section interfaces for decorations that provide their own
+ * sections (via IMovieSceneSectionProviderDecoration), add entries to section context menus,
+ * and provide custom key editor and keyframe widgets for decoration channels.
  *
  * Example registration in an editor module's StartupModule():
  *
  * class FMyDecorationEditor : public ISequencerDecorationEditor
  * {
  *     virtual UClass* GetDecorationClass() const override { return UMyDecoration::StaticClass(); }
- *     virtual TSharedRef<ISequencerSection> MakeSectionInterface(...) override { return MakeShared<FMySectionInterface>(...); }
+ *     virtual TSharedPtr<ISequencerSection> MakeSectionInterface(...) override { return MakeShared<FMySectionInterface>(...); }
  * };
  *
  * IModularFeatures::Get().RegisterModularFeature(ISequencerDecorationEditor::GetModularFeatureName(), MyDecorationEditorPtr);
@@ -43,10 +47,56 @@ public:
 	 * Create a section interface for a section provided by this decoration.
 	 * The section interface is used for displaying/editing the section in Sequencer.
 	 */
-	virtual TSharedRef<ISequencerSection> MakeSectionInterface(
+	virtual TSharedPtr<ISequencerSection> MakeSectionInterface(
 		UMovieSceneSection& Section,
 		UMovieSceneTrack& OwningTrack,
-		FGuid ObjectBinding) = 0;
+		FGuid ObjectBinding) { return nullptr; }
+
+	/**
+	 * Add entries to the section's right-click context menu on behalf of this decoration.
+	 * Called for each decoration on the section that has a registered editor.
+	 */
+	virtual void BuildSectionContextMenu(
+		FMenuBuilder& MenuBuilder,
+		UMovieSceneSection& Section,
+		UObject& Decoration,
+		const FGuid& ObjectBinding,
+		TWeakPtr<ISequencer> Sequencer) {}
+
+	/**
+	 * Optionally create a custom key editor widget for one of this decoration's channels.
+	 * Return nullptr to use the standard channel editor.
+	 */
+	virtual TSharedPtr<SWidget> CreateKeyEditor(
+		UObject& Decoration,
+		UMovieSceneSection& Section,
+		FMovieSceneChannelHandle ChannelHandle,
+		TWeakPtr<ISequencer> Sequencer,
+		const FGuid& ObjectBindingID) { return nullptr; }
+
+	/**
+	 * Called when a section model is first created for a section that has this decoration.
+	 * Override to set up per-section subscriptions (e.g., channel key events).
+	 */
+	virtual void OnSectionInterfaceCreated(
+		UMovieSceneSection& Section,
+		UObject& Decoration,
+		TWeakPtr<ISequencer> Sequencer) {}
+
+	/**
+	 * Optionally create a custom widget for the KeyFrame outliner column
+	 * (the "add key" button) for one of this decoration's channels.
+	 * The ChannelGroupModel is the outliner view model for the channel group -
+	 * pass it to SKeyNavigationButtons to get correct row hover behavior.
+	 * Return nullptr to use the standard add-key button.
+	 */
+	virtual TSharedPtr<SWidget> CreateKeyFrameWidget(
+		UObject& Decoration,
+		UMovieSceneSection& Section,
+		FMovieSceneChannelHandle ChannelHandle,
+		TWeakPtr<ISequencer> Sequencer,
+		const FGuid& ObjectBindingID,
+		TSharedPtr<UE::Sequencer::FViewModel> ChannelGroupModel) { return nullptr; }
 
 	/**
 	 * Get the modular feature name for decoration editors.
diff --git a/Engine/Source/Editor/Sequencer/Public/MVVM/SectionModelStorageExtension.h b/Engine/Source/Editor/Sequencer/Public/MVVM/SectionModelStorageExtension.h
--- a/Engine/Source/Editor/Sequencer/Public/MVVM/SectionModelStorageExtension.h
+++ b/Engine/Source/Editor/Sequencer/Public/MVVM/SectionModelStorageExtension.h
@@ -34,10 +34,14 @@ public:
 
 	UE_API TSharedPtr<FSectionModel> FindModelForSection(const UMovieSceneSection* InSection) const;
 
+	UE_API virtual void OnCreated(TSharedRef<FViewModel> InOwner) override;
 	UE_API virtual void OnReinitialize() override;
 
 private:
 
+	void NotifyDecorationEditors(UMovieSceneSection* InSection);
+
+	TWeakPtr<FViewModel> WeakOwner;
 	TMap<FObjectKey, TWeakPtr<FSectionModel>> SectionToModel;
 };
 
diff --git a/Engine/Source/Runtime/MovieScene/Private/EntitySystem/MovieSceneEntitySystemRunner.cpp b/Engine/Source/Runtime/MovieScene/Private/EntitySystem/MovieSceneEntitySystemRunner.cpp
--- a/Engine/Source/Runtime/MovieScene/Private/EntitySystem/MovieSceneEntitySystemRunner.cpp
+++ b/Engine/Source/Runtime/MovieScene/Private/EntitySystem/MovieSceneEntitySystemRunner.cpp
@@ -520,6 +520,11 @@ void FMovieSceneEntitySystemRunner::FlushOutstanding(double BudgetMs, UE::MovieS
 		return;
 	}
 
+	// Fire PreFlush before any phases run, allowing systems to perform
+	// work that must happen outside the ECS pipeline (e.g. rebuilding
+	// caches that require re-entrant evaluation).
+	Linker->Events.PreFlush.Broadcast(Linker);
+
 	SCOPE_CYCLE_COUNTER(MovieSceneEval_RunnerFlush);
 
 	// We need to run the system from the game thread so we know we can fire events and callbacks from here.
diff --git a/Engine/Source/Runtime/MovieScene/Public/Decorations/IMovieSceneChannelDecoration.h b/Engine/Source/Runtime/MovieScene/Public/Decorations/IMovieSceneChannelDecoration.h
--- a/Engine/Source/Runtime/MovieScene/Public/Decorations/IMovieSceneChannelDecoration.h
+++ b/Engine/Source/Runtime/MovieScene/Public/Decorations/IMovieSceneChannelDecoration.h
@@ -17,7 +17,7 @@ public:
 };
 
 
-/** 
+/**
  * Optional decoration that can be added to sections to add channels
  */
 class IMovieSceneChannelDecoration
@@ -33,5 +33,20 @@ public:
 	{
 		return EMovieSceneChannelProxyType::Static;
 	}
+
+	// Mark this decoration's channel proxy as needing reconstruction.
+	// Call when internal state changes which channels PopulateChannelProxy returns.
+	void InvalidateDecorationChannelProxy() { bDecorationChannelProxyDirty = true; }
+
+	// Returns true and clears the flag if the channel proxy needs rebuilding.
+	bool ConsumeDecorationChannelProxyDirty()
+	{
+		bool bWasDirty = bDecorationChannelProxyDirty;
+		bDecorationChannelProxyDirty = false;
+		return bWasDirty;
+	}
+
+private:
+	bool bDecorationChannelProxyDirty = false;
 };
 
diff --git a/Engine/Source/Runtime/MovieScene/Public/EntitySystem/MovieSceneEntitySystemLinker.h b/Engine/Source/Runtime/MovieScene/Public/EntitySystem/MovieSceneEntitySystemLinker.h
--- a/Engine/Source/Runtime/MovieScene/Public/EntitySystem/MovieSceneEntitySystemLinker.h
+++ b/Engine/Source/Runtime/MovieScene/Public/EntitySystem/MovieSceneEntitySystemLinker.h
@@ -384,6 +384,7 @@ public:
 
 	struct
 	{
+		FMovieSceneEntitySystemLinkerEvent          PreFlush;
 		FMovieSceneEntitySystemLinkerPostSpawnEvent PostSpawnEvent;
 		FMovieSceneEntitySystemLinkerEvent          TagGarbage;
 		FMovieSceneEntitySystemLinkerEvent          CleanTaggedGarbage;
diff --git a/Engine/Source/Runtime/MovieSceneTracks/Private/Tracks/MovieSceneCommonAnimationTrack.cpp b/Engine/Source/Runtime/MovieSceneTracks/Private/Tracks/MovieSceneCommonAnimationTrack.cpp
--- a/Engine/Source/Runtime/MovieSceneTracks/Private/Tracks/MovieSceneCommonAnimationTrack.cpp
+++ b/Engine/Source/Runtime/MovieSceneTracks/Private/Tracks/MovieSceneCommonAnimationTrack.cpp
@@ -3,6 +3,7 @@
 #include "Tracks/MovieSceneCommonAnimationTrack.h"
 
 #include "Animation/AnimationPoseData.h"
+#include "MovieSceneAnimationRotationUtils.h"
 #include "Animation/AnimInstance.h"
 #include "Animation/AnimSequence.h"
 #include "Animation/AttributesRuntime.h"
@@ -727,228 +728,10 @@ static FTransform GetTransformForBoneRelativeToIndex(UAnimSequence* AnimSequence
 	return WorldTransform;
 }
 
-enum class ESkelAnimRotationOrder 
-{
-	XYZ,
-	XZY,
-	YXZ,
-	YZX,
-	ZXY,
-	ZYX
-};
-
-static FQuat QuatFromEuler(const FVector& XYZAnglesInDegrees, ESkelAnimRotationOrder RotationOrder)
-{
-	float X = FMath::DegreesToRadians(XYZAnglesInDegrees.X);
-	float Y = FMath::DegreesToRadians(XYZAnglesInDegrees.Y);
-	float Z = FMath::DegreesToRadians(XYZAnglesInDegrees.Z);
-
-	float CosX = FMath::Cos(X * 0.5f);
-	float CosY = FMath::Cos(Y * 0.5f);
-	float CosZ = FMath::Cos(Z * 0.5f);
-
-	float SinX = FMath::Sin(X * 0.5f);
-	float SinY = FMath::Sin(Y * 0.5f);
-	float SinZ = FMath::Sin(Z * 0.5f);
-
-	if (RotationOrder == ESkelAnimRotationOrder::XYZ)
-	{
-		return FQuat(SinX * CosY * CosZ - CosX * SinY * SinZ,
-			CosX * SinY * CosZ + SinX * CosY * SinZ,
-			CosX * CosY * SinZ - SinX * SinY * CosZ,
-			CosX * CosY * CosZ + SinX * SinY * SinZ);
-
-	}
-	else if (RotationOrder == ESkelAnimRotationOrder::XZY)
-	{
-		return FQuat(SinX * CosY * CosZ + CosX * SinY * SinZ,
-			CosX * SinY * CosZ + SinX * CosY * SinZ,
-			CosX * CosY * SinZ - SinX * SinY * CosZ,
-			CosX * CosY * CosZ - SinX * SinY * SinZ);
-
-	}
-	else if (RotationOrder == ESkelAnimRotationOrder::YXZ)
-	{
-		return FQuat(SinX * CosY * CosZ - CosX * SinY * SinZ,
-			CosX * SinY * CosZ + SinX * CosY * SinZ,
-			CosX * CosY * SinZ + SinX * SinY * CosZ,
-			CosX * CosY * CosZ - SinX * SinY * SinZ);
-
-	}
-	else if (RotationOrder == ESkelAnimRotationOrder::YZX)
-	{
-		return FQuat(SinX * CosY * CosZ - CosX * SinY * SinZ,
-			CosX * SinY * CosZ - SinX * CosY * SinZ,
-			CosX * CosY * SinZ + SinX * SinY * CosZ,
-			CosX * CosY * CosZ + SinX * SinY * SinZ);
-	}
-	else if (RotationOrder == ESkelAnimRotationOrder::ZXY)
-	{
-		return FQuat(SinX * CosY * CosZ + CosX * SinY * SinZ,
-			CosX * SinY * CosZ - SinX * CosY * SinZ,
-			CosX * CosY * SinZ - SinX * SinY * CosZ,
-			CosX * CosY * CosZ + SinX * SinY * SinZ);
-
-	}
-	else if (RotationOrder == ESkelAnimRotationOrder::ZYX)
-	{
-		return FQuat(SinX * CosY * CosZ + CosX * SinY * SinZ,
-			CosX * SinY * CosZ - SinX * CosY * SinZ,
-			CosX * CosY * SinZ + SinX * SinY * CosZ,
-			CosX * CosY * CosZ - SinX * SinY * SinZ);
-
-	}
-
-	// should not happen
-	return FQuat::Identity;
-}
-
-static FVector EulerFromQuat(const FQuat& Rotation, ESkelAnimRotationOrder RotationOrder)
-{
-	float X = Rotation.X;
-	float Y = Rotation.Y;
-	float Z = Rotation.Z;
-	float W = Rotation.W;
-	float X2 = X * 2.f;
-	float Y2 = Y * 2.f;
-	float Z2 = Z * 2.f;
-	float XX2 = X * X2;
-	float XY2 = X * Y2;
-	float XZ2 = X * Z2;
-	float YX2 = Y * X2;
-	float YY2 = Y * Y2;
-	float YZ2 = Y * Z2;
-	float ZX2 = Z * X2;
-	float ZY2 = Z * Y2;
-	float ZZ2 = Z * Z2;
-	float WX2 = W * X2;
-	float WY2 = W * Y2;
-	float WZ2 = W * Z2;
-
-	FVector AxisX, AxisY, AxisZ;
-	AxisX.X = (1.f - (YY2 + ZZ2));
-	AxisY.X = (XY2 + WZ2);
-	AxisZ.X = (XZ2 - WY2);
-	AxisX.Y = (XY2 - WZ2);
-	AxisY.Y = (1.f - (XX2 + ZZ2));
-	AxisZ.Y = (YZ2 + WX2);
-	AxisX.Z = (XZ2 + WY2);
-	AxisY.Z = (YZ2 - WX2);
-	AxisZ.Z = (1.f - (XX2 + YY2));
-
-	FVector Result = FVector::ZeroVector;
-
-	if (RotationOrder == ESkelAnimRotationOrder::XYZ)
-	{
-		Result.Y = FMath::Asin(-FMath::Clamp<float>(AxisZ.X, -1.f, 1.f));
-
-		if (FMath::Abs(AxisZ.X) < 1.f - SMALL_NUMBER)
-		{
-			Result.X = FMath::Atan2(AxisZ.Y, AxisZ.Z);
-			Result.Z = FMath::Atan2(AxisY.X, AxisX.X);
-		}
-		else
-		{
-			Result.X = 0.f;
-			Result.Z = FMath::Atan2(-AxisX.Y, AxisY.Y);
-		}
-	}
-	else if (RotationOrder == ESkelAnimRotationOrder::XZY)
-	{
-
-		Result.Z = FMath::Asin(FMath::Clamp<float>(AxisY.X, -1.f, 1.f));
-
-		if (FMath::Abs(AxisY.X) < 1.f - SMALL_NUMBER)
-		{
-			Result.X = FMath::Atan2(-AxisY.Z, AxisY.Y);
-			Result.Y = FMath::Atan2(-AxisZ.X, AxisX.X);
-		}
-		else
-		{
-			Result.X = 0.f;
-			Result.Y = FMath::Atan2(AxisX.Z, AxisZ.Z);
-		}
-	}
-	else if (RotationOrder == ESkelAnimRotationOrder::YXZ)
-	{
-		Result.X = FMath::Asin(FMath::Clamp<float>(AxisZ.Y, -1.f, 1.f));
-
-		if (FMath::Abs(AxisZ.Y) < 1.f - SMALL_NUMBER)
-		{
-			Result.Y = FMath::Atan2(-AxisZ.X, AxisZ.Z);
-			Result.Z = FMath::Atan2(-AxisX.Y, AxisY.Y);
-		}
-		else
-		{
-			Result.Y = 0.f;
-			Result.Z = FMath::Atan2(AxisY.X, AxisX.X);
-		}
-	}
-	else if (RotationOrder == ESkelAnimRotationOrder::YZX)
-	{
-		Result.Z = FMath::Asin(-FMath::Clamp<float>(AxisX.Y, -1.f, 1.f));
-
-		if (FMath::Abs(AxisX.Y) < 1.f - SMALL_NUMBER)
-		{
-			Result.X = FMath::Atan2(AxisZ.Y, AxisY.Y);
-			Result.Y = FMath::Atan2(AxisX.Z, AxisX.X);
-		}
-		else
-		{
-			Result.X = FMath::Atan2(-AxisY.Z, AxisZ.Z);
-			Result.Y = 0.f;
-		}
-	}
-	else if (RotationOrder == ESkelAnimRotationOrder::ZXY)
-	{
-		Result.X = FMath::Asin(-FMath::Clamp<float>(AxisY.Z, -1.f, 1.f));
-
-		if (FMath::Abs(AxisY.Z) < 1.f - SMALL_NUMBER)
-		{
-			Result.Y = FMath::Atan2(AxisX.Z, AxisZ.Z);
-			Result.Z = FMath::Atan2(AxisY.X, AxisY.Y);
-		}
-		else
-		{
-			Result.Y = FMath::Atan2(-AxisZ.X, AxisX.X);
-			Result.Z = 0.f;
-		}
-	}
-	else if (RotationOrder == ESkelAnimRotationOrder::ZYX)
-	{
-		Result.Y = FMath::Asin(FMath::Clamp<float>(AxisX.Z, -1.f, 1.f));
-
-		if (FMath::Abs(AxisX.Z) < 1.f - SMALL_NUMBER)
-		{
-			Result.X = FMath::Atan2(-AxisY.Z, AxisZ.Z);
-			Result.Z = FMath::Atan2(-AxisX.Y, AxisX.X);
-		}
-		else
-		{
-			Result.X = FMath::Atan2(AxisZ.Y, AxisY.Y);
-			Result.Z = 0.f;
-		}
-	}
-
-	return Result * 180.f / PI;
-}
-/**
-*  Function to find best rotation order given how we are matching the rotations. 
-*  If it is matched we need to make sure it happens first
-*  Issue is Yaw is most common match but by default FRotation:FQuat conversions it's last, this causes issues.
-*/
-static ESkelAnimRotationOrder FindBestRotationOrder(bool bMatchRoll, bool bMatchPitch, bool bMatchYaw)
-{
-	if (bMatchYaw)
-	{
-		return ESkelAnimRotationOrder::YXZ;
-	}
-	if (bMatchPitch)
-	{
-		return ESkelAnimRotationOrder::YZX;
-	}
-	return ESkelAnimRotationOrder::XYZ;
-}
+using ERotationOrder = UE::MovieScene::ERotationOrder;
+using UE::MovieScene::QuatFromEuler;
+using UE::MovieScene::EulerFromQuat;
+using UE::MovieScene::FindBestRotationOrder;
 
 void UMovieSceneCommonAnimationTrack::MatchSectionByBoneTransform(bool bMatchWithPrevious, USkeletalMeshComponent* SkelMeshComp, UMovieSceneSkeletalAnimationSection* CurrentSection, FFrameTime CurrentFrame, FFrameRate InTickResolution,
 	const FName& BoneName, FTransform& SecondSectionRootDiff, FVector& TranslationDiff, FQuat& RotationDiff) //add options for z and for rotation.
@@ -993,7 +776,7 @@ void UMovieSceneCommonAnimationTrack::MatchSectionByBoneTransform(bool bMatchWit
 			//Need to match the translations and rotations here 
 			//First need to get the correct rotation order based upon what's matching, otherwise if not all are matched 
 			//and one rotation is set last we will get errors.
-			ESkelAnimRotationOrder RotationOrder = FindBestRotationOrder(CurrentSection->bMatchRotationRoll, CurrentSection->bMatchRotationPitch, CurrentSection->bMatchRotationYaw);
+			ERotationOrder RotationOrder = FindBestRotationOrder(CurrentSection->bMatchRotationRoll, CurrentSection->bMatchRotationPitch, CurrentSection->bMatchRotationYaw);
 
 			FVector FirstTransformTranslation = FirstTransform.GetTranslation();
 			FVector SecondTransformTranslation = SecondTransform.GetTranslation();
@@ -1060,7 +843,7 @@ void UMovieSceneCommonAnimationTrack::MatchSectionByBoneTransform(bool bMatchWit
 			//Need to match the translations and rotations here 
 			//First need to get the correct rotation order based upon what's matching, otherwise if not all are matched 
 			//and one rotation is set last we will get errors.
-			ESkelAnimRotationOrder RotationOrder = FindBestRotationOrder(CurrentSection->bMatchRotationRoll, CurrentSection->bMatchRotationPitch, CurrentSection->bMatchRotationYaw);
+			ERotationOrder RotationOrder = FindBestRotationOrder(CurrentSection->bMatchRotationRoll, CurrentSection->bMatchRotationPitch, CurrentSection->bMatchRotationYaw);
 
 			FVector FirstTransformTranslation = FirstTransform.GetTranslation();
 			FVector SecondTransformTranslation = SecondTransform.GetTranslation();
diff --git a/Engine/Source/Runtime/MovieSceneTracks/Public/MovieSceneAnimationRotationUtils.h b/Engine/Source/Runtime/MovieSceneTracks/Public/MovieSceneAnimationRotationUtils.h
new file mode 100644
--- /dev/null
+++ b/Engine/Source/Runtime/MovieSceneTracks/Public/MovieSceneAnimationRotationUtils.h
@@ -0,0 +1,207 @@
+// Copyright Epic Games, Inc. All Rights Reserved.
+
+#pragma once
+
+#include "CoreTypes.h"
+#include "Math/Quat.h"
+#include "Math/Rotator.h"
+#include "Math/Vector.h"
+
+namespace UE::MovieScene
+{
+
+// Euler decomposition order for rotation matching. When matching specific
+// rotation axes, the matched axis should be decomposed first to avoid
+// gimbal lock artifacts on the axes we care about.
+enum class ERotationOrder
+{
+	XYZ, XZY, YXZ, YZX, ZXY, ZYX
+};
+
+// Choose the decomposition order that puts the matched axis first.
+// Yaw (Z) is the most common match axis, so YXZ is the default when matching yaw.
+inline ERotationOrder FindBestRotationOrder(bool bMatchRoll, bool bMatchPitch, bool bMatchYaw)
+{
+	if (bMatchYaw)
+	{
+		return ERotationOrder::YXZ;
+	}
+	if (bMatchPitch)
+	{
+		return ERotationOrder::YZX;
+	}
+	return ERotationOrder::XYZ;
+}
+
+inline FQuat QuatFromEuler(const FVector& XYZAnglesInDegrees, ERotationOrder RotationOrder)
+{
+	double X = FMath::DegreesToRadians(XYZAnglesInDegrees.X);
+	double Y = FMath::DegreesToRadians(XYZAnglesInDegrees.Y);
+	double Z = FMath::DegreesToRadians(XYZAnglesInDegrees.Z);
+
+	double CosX = FMath::Cos(X * 0.5);
+	double CosY = FMath::Cos(Y * 0.5);
+	double CosZ = FMath::Cos(Z * 0.5);
+
+	double SinX = FMath::Sin(X * 0.5);
+	double SinY = FMath::Sin(Y * 0.5);
+	double SinZ = FMath::Sin(Z * 0.5);
+
+	switch (RotationOrder)
+	{
+	case ERotationOrder::XYZ:
+		return FQuat(SinX * CosY * CosZ - CosX * SinY * SinZ,
+			CosX * SinY * CosZ + SinX * CosY * SinZ,
+			CosX * CosY * SinZ - SinX * SinY * CosZ,
+			CosX * CosY * CosZ + SinX * SinY * SinZ);
+
+	case ERotationOrder::XZY:
+		return FQuat(SinX * CosY * CosZ + CosX * SinY * SinZ,
+			CosX * SinY * CosZ + SinX * CosY * SinZ,
+			CosX * CosY * SinZ - SinX * SinY * CosZ,
+			CosX * CosY * CosZ - SinX * SinY * SinZ);
+
+	case ERotationOrder::YXZ:
+		return FQuat(SinX * CosY * CosZ - CosX * SinY * SinZ,
+			CosX * SinY * CosZ + SinX * CosY * SinZ,
+			CosX * CosY * SinZ + SinX * SinY * CosZ,
+			CosX * CosY * CosZ - SinX * SinY * SinZ);
+
+	case ERotationOrder::YZX:
+		return FQuat(SinX * CosY * CosZ - CosX * SinY * SinZ,
+			CosX * SinY * CosZ - SinX * CosY * SinZ,
+			CosX * CosY * SinZ + SinX * SinY * CosZ,
+			CosX * CosY * CosZ + SinX * SinY * SinZ);
+
+	case ERotationOrder::ZXY:
+		return FQuat(SinX * CosY * CosZ + CosX * SinY * SinZ,
+			CosX * SinY * CosZ - SinX * CosY * SinZ,
+			CosX * CosY * SinZ - SinX * SinY * CosZ,
+			CosX * CosY * CosZ + SinX * SinY * SinZ);
+
+	case ERotationOrder::ZYX:
+		return FQuat(SinX * CosY * CosZ + CosX * SinY * SinZ,
+			CosX * SinY * CosZ - SinX * CosY * SinZ,
+			CosX * CosY * SinZ + SinX * SinY * CosZ,
+			CosX * CosY * CosZ - SinX * SinY * SinZ);
+	}
+
+	return FQuat::Identity;
+}
+
+inline FVector EulerFromQuat(const FQuat& Rotation, ERotationOrder RotationOrder)
+{
+	double X = Rotation.X;
+	double Y = Rotation.Y;
+	double Z = Rotation.Z;
+	double W = Rotation.W;
+	double X2 = X * 2.0;
+	double Y2 = Y * 2.0;
+	double Z2 = Z * 2.0;
+	double XX2 = X * X2;
+	double XY2 = X * Y2;
+	double XZ2 = X * Z2;
+	double YY2 = Y * Y2;
+	double YZ2 = Y * Z2;
+	double ZZ2 = Z * Z2;
+	double WX2 = W * X2;
+	double WY2 = W * Y2;
+	double WZ2 = W * Z2;
+
+	FVector AxisX(1.0 - (YY2 + ZZ2), XY2 + WZ2, XZ2 - WY2);
+	FVector AxisY(XY2 - WZ2, 1.0 - (XX2 + ZZ2), YZ2 + WX2);
+	FVector AxisZ(XZ2 + WY2, YZ2 - WX2, 1.0 - (XX2 + YY2));
+
+	FVector Result = FVector::ZeroVector;
+
+	switch (RotationOrder)
+	{
+	case ERotationOrder::XYZ:
+		Result.Y = FMath::Asin(-FMath::Clamp(AxisZ.X, -1.0, 1.0));
+		if (FMath::Abs(AxisZ.X) < 1.0 - SMALL_NUMBER)
+		{
+			Result.X = FMath::Atan2(AxisZ.Y, AxisZ.Z);
+			Result.Z = FMath::Atan2(AxisY.X, AxisX.X);
+		}
+		else
+		{
+			Result.X = 0.0;
+			Result.Z = FMath::Atan2(-AxisX.Y, AxisY.Y);
+		}
+		break;
+
+	case ERotationOrder::XZY:
+		Result.Z = FMath::Asin(FMath::Clamp(AxisY.X, -1.0, 1.0));
+		if (FMath::Abs(AxisY.X) < 1.0 - SMALL_NUMBER)
+		{
+			Result.X = FMath::Atan2(-AxisY.Z, AxisY.Y);
+			Result.Y = FMath::Atan2(-AxisZ.X, AxisX.X);
+		}
+		else
+		{
+			Result.X = 0.0;
+			Result.Y = FMath::Atan2(AxisX.Z, AxisZ.Z);
+		}
+		break;
+
+	case ERotationOrder::YXZ:
+		Result.X = FMath::Asin(FMath::Clamp(AxisZ.Y, -1.0, 1.0));
+		if (FMath::Abs(AxisZ.Y) < 1.0 - SMALL_NUMBER)
+		{
+			Result.Y = FMath::Atan2(-AxisZ.X, AxisZ.Z);
+			Result.Z = FMath::Atan2(-AxisX.Y, AxisY.Y);
+		}
+		else
+		{
+			Result.Y = 0.0;
+			Result.Z = FMath::Atan2(AxisY.X, AxisX.X);
+		}
+		break;
+
+	case ERotationOrder::YZX:
+		Result.Z = FMath::Asin(-FMath::Clamp(AxisX.Y, -1.0, 1.0));
+		if (FMath::Abs(AxisX.Y) < 1.0 - SMALL_NUMBER)
+		{
+			Result.X = FMath::Atan2(AxisZ.Y, AxisY.Y);
+			Result.Y = FMath::Atan2(AxisX.Z, AxisX.X);
+		}
+		else
+		{
+			Result.X = FMath::Atan2(-AxisY.Z, AxisZ.Z);
+			Result.Y = 0.0;
+		}
+		break;
+
+	case ERotationOrder::ZXY:
+		Result.X = FMath::Asin(-FMath::Clamp(AxisY.Z, -1.0, 1.0));
+		if (FMath::Abs(AxisY.Z) < 1.0 - SMALL_NUMBER)
+		{
+			Result.Y = FMath::Atan2(AxisX.Z, AxisZ.Z);
+			Result.Z = FMath::Atan2(AxisY.X, AxisY.Y);
+		}
+		else
+		{
+			Result.Y = FMath::Atan2(-AxisZ.X, AxisX.X);
+			Result.Z = 0.0;
+		}
+		break;
+
+	case ERotationOrder::ZYX:
+		Result.Y = FMath::Asin(FMath::Clamp(AxisX.Z, -1.0, 1.0));
+		if (FMath::Abs(AxisX.Z) < 1.0 - SMALL_NUMBER)
+		{
+			Result.X = FMath::Atan2(-AxisY.Z, AxisZ.Z);
+			Result.Z = FMath::Atan2(-AxisX.Y, AxisX.X);
+		}
+		else
+		{
+			Result.X = FMath::Atan2(AxisZ.Y, AxisY.Y);
+			Result.Z = 0.0;
+		}
+		break;
+	}
+
+	return Result * (180.0 / UE_DOUBLE_PI);
+}
+
+} // namespace UE::MovieScene
