Sequencer: Mixer bugfixes.
* Fix re-entrancy crash
* Fix various issues with cut/paste of sections into the mixer

#Synced-CL 51336034

diff --git a/Engine/Plugins/Experimental/MovieSceneAnimMixer/Source/MovieSceneAnimMixer/Private/MovieSceneAnimationMixerTrack.cpp b/Engine/Plugins/Experimental/MovieSceneAnimMixer/Source/MovieSceneAnimMixer/Private/MovieSceneAnimationMixerTrack.cpp
--- a/Engine/Plugins/Experimental/MovieSceneAnimMixer/Source/MovieSceneAnimMixer/Private/MovieSceneAnimationMixerTrack.cpp
+++ b/Engine/Plugins/Experimental/MovieSceneAnimMixer/Source/MovieSceneAnimMixer/Private/MovieSceneAnimationMixerTrack.cpp
@@ -305,7 +305,12 @@ UMovieSceneAnimationMixerTrack::UMovieSceneAnimationMixerTrack(const FObjectInit
 
 bool UMovieSceneAnimationMixerTrack::SupportsType(TSubclassOf<UMovieSceneSection> SectionClass) const
 {
-	return SectionClass.Get()->ImplementsInterface(UMovieSceneAnimationMixerItemInterface::StaticClass());
+	if (SectionClass.Get()->IsChildOf(UMovieSceneAnimTransitionSectionBase::StaticClass()))
+	{
+		return false;
+	}
+	return Super::SupportsType(SectionClass)
+		|| SectionClass.Get()->ImplementsInterface(UMovieSceneAnimationMixerItemInterface::StaticClass());
 }
 
 bool UMovieSceneAnimationMixerTrack::FixRowIndices()
@@ -330,20 +335,29 @@ void UMovieSceneAnimationMixerTrack::OnSectionAddedImpl(UMovieSceneSection* Sect
 {
 	if (Section)
 	{
-		UMovieSceneAnimationSectionDecoration* NewDecoration = NewObject<UMovieSceneAnimationSectionDecoration>(Section, UMovieSceneAnimationSectionDecoration::StaticClass(), NAME_None, RF_Transactional);
-		NewDecoration->MixerTrack = this;
-		Section->AddDecoration(NewDecoration);
+		UMovieSceneAnimationSectionDecoration* ExistingDecoration = Section->FindDecoration<UMovieSceneAnimationSectionDecoration>();
+		if (!ExistingDecoration)
+		{
+			UMovieSceneAnimationSectionDecoration* NewDecoration = NewObject<UMovieSceneAnimationSectionDecoration>(Section, UMovieSceneAnimationSectionDecoration::StaticClass(), NAME_None, RF_Transactional);
+			NewDecoration->MixerTrack = this;
+			Section->AddDecoration(NewDecoration);
+		}
+		else
+		{
+			ExistingDecoration->MixerTrack = this;
+		}
 
-		// Auto-add Root Motion Settings decoration only for skeletal animation sections
-		// whose animations have root motion enabled
 		if (UMovieSceneSkeletalAnimationSection* AnimSection = Cast<UMovieSceneSkeletalAnimationSection>(Section))
 		{
-			if (UAnimSequenceBase* Animation = AnimSection->Params.Animation)
+			if (!AnimSection->FindDecoration<UMovieSceneRootMotionSettingsDecoration>())
 			{
-				if (Animation->HasRootMotion())
+				if (UAnimSequenceBase* Animation = AnimSection->Params.Animation)
 				{
-					UMovieSceneRootMotionSettingsDecoration* RootMotionSettings = NewObject<UMovieSceneRootMotionSettingsDecoration>(AnimSection, UMovieSceneRootMotionSettingsDecoration::StaticClass(), NAME_None, RF_Transactional);
-					AnimSection->AddDecoration(RootMotionSettings);
+					if (Animation->HasRootMotion())
+					{
+						UMovieSceneRootMotionSettingsDecoration* RootMotionSettings = NewObject<UMovieSceneRootMotionSettingsDecoration>(AnimSection, UMovieSceneRootMotionSettingsDecoration::StaticClass(), NAME_None, RF_Transactional);
+						AnimSection->AddDecoration(RootMotionSettings);
+					}
 				}
 			}
 		}
diff --git a/Engine/Plugins/Experimental/MovieSceneAnimMixer/Source/MovieSceneAnimMixer/Private/Systems/MovieSceneAnimMixerSystem.cpp b/Engine/Plugins/Experimental/MovieSceneAnimMixer/Source/MovieSceneAnimMixer/Private/Systems/MovieSceneAnimMixerSystem.cpp
--- a/Engine/Plugins/Experimental/MovieSceneAnimMixer/Source/MovieSceneAnimMixer/Private/Systems/MovieSceneAnimMixerSystem.cpp
+++ b/Engine/Plugins/Experimental/MovieSceneAnimMixer/Source/MovieSceneAnimMixer/Private/Systems/MovieSceneAnimMixerSystem.cpp
@@ -1091,27 +1091,53 @@ void UMovieSceneAnimMixerSystem::OnUnlink()
 
 void UMovieSceneAnimMixerSystem::HandlePreFlush(UMovieSceneEntitySystemLinker* InLinker)
 {
+	using namespace UE::MovieScene;
+
+	if (bInsidePreFlush)
+	{
+		return;
+	}
+	TGuardValue<bool> ReentrancyGuard(bInsidePreFlush, true);
+
+	// Collect dirty tracks before rebuilding. The rebuild triggers re-entrant
+	// flushes (via bake evaluation) which can add/remove Mixers map entries,
+	// so we must not iterate the map during the rebuild.
+	TArray<TWeakObjectPtr<UMovieSceneAnimationMixerTrack>> DirtyTracks;
+	TArray<FInstanceHandle> DirtyHandles;
+
 	for (TPair<FMovieSceneAnimMixerKey, TSharedPtr<FMovieSceneAnimMixer>>& Pair : Mixers)
 	{
+		if (!Pair.Value)
+		{
+			continue;
+		}
 		UMovieSceneAnimationMixerTrack* Track = Pair.Value->OwnerTrack.Get();
 		if (!Track || !Track->HasDirtyAccumulatedOffsetCache())
 		{
 			continue;
 		}
 
-		// Find a valid instance handle from any live entry in this mixer
 		for (const TWeakPtr<FMovieSceneAnimMixerEntry>& WeakEntry : Pair.Value->WeakEntries)
 		{
 			if (TSharedPtr<FMovieSceneAnimMixerEntry> Entry = WeakEntry.Pin())
 			{
 				if (Entry->InstanceHandle.IsValid())
 				{
-					Track->RebuildDirtyAccumulatedOffsetCache(InLinker, Entry->InstanceHandle);
+					DirtyTracks.Add(Track);
+					DirtyHandles.Add(Entry->InstanceHandle);
 					break;
 				}
 			}
 		}
 	}
+
+	for (int32 Index = 0; Index < DirtyTracks.Num(); ++Index)
+	{
+		if (UMovieSceneAnimationMixerTrack* Track = DirtyTracks[Index].Get())
+		{
+			Track->RebuildDirtyAccumulatedOffsetCache(InLinker, DirtyHandles[Index]);
+		}
+	}
 }
 
 void UMovieSceneAnimMixerSystem::OnRun(FSystemTaskPrerequisites& InPrerequisites, FSystemSubsequentTasks& Subsequents)
diff --git a/Engine/Plugins/Experimental/MovieSceneAnimMixer/Source/MovieSceneAnimMixer/Public/Systems/MovieSceneAnimMixerSystem.h b/Engine/Plugins/Experimental/MovieSceneAnimMixer/Source/MovieSceneAnimMixer/Public/Systems/MovieSceneAnimMixerSystem.h
--- a/Engine/Plugins/Experimental/MovieSceneAnimMixer/Source/MovieSceneAnimMixer/Public/Systems/MovieSceneAnimMixerSystem.h
+++ b/Engine/Plugins/Experimental/MovieSceneAnimMixer/Source/MovieSceneAnimMixer/Public/Systems/MovieSceneAnimMixerSystem.h
@@ -207,4 +207,5 @@ private:
 	TObjectPtr<UMovieSceneRootMotionSystem> RootMotionSystem;
 
 	FDelegateHandle PreFlushHandle;
+	bool bInsidePreFlush = false;
 };
\ No newline at end of file
diff --git a/Engine/Source/Editor/Sequencer/Private/SequencerUtilities.cpp b/Engine/Source/Editor/Sequencer/Private/SequencerUtilities.cpp
--- a/Engine/Source/Editor/Sequencer/Private/SequencerUtilities.cpp
+++ b/Engine/Source/Editor/Sequencer/Private/SequencerUtilities.cpp
@@ -2366,7 +2366,13 @@ bool FSequencerUtilities::PasteSections(const FString& TextToImport, FMovieScene
 			if (Section->HasStartFrame())
 			{
 				FFrameNumber NewStartFrame = PasteSectionsParams.Time.FrameNumber + (Section->GetInclusiveStartFrame() - FirstFrame.GetValue());
-				Section->MoveSection(NewStartFrame - Section->GetInclusiveStartFrame());
+				FFrameNumber Delta = NewStartFrame - Section->GetInclusiveStartFrame();
+				if (Delta != 0)
+				{
+					Section->MoveSection(Delta);
+					FMovieSceneSectionMovedParams MovedParams(EPropertyChangeType::ValueSet);
+					Track->OnSectionMoved(*Section, MovedParams);
+				}
 			}
 
 			if (!bAllowOverlap)
@@ -2384,33 +2390,7 @@ bool FSequencerUtilities::PasteSections(const FString& TextToImport, FMovieScene
 		}
 
 		// Fix up rows after sections are in place
-		if (Track->SupportsMultipleRows())
-		{
-			// If any newly created section overlaps the previous sections, put all the sections on the max available row
-			// Find the  this section overlaps any previous sections, 
-			int32 MaxAvailableRowIndex = -1;
-			for (UMovieSceneSection* Section : OutSections)
-			{
-				if (!Track->SupportsType(Section->GetClass()))
-				{
-					continue;
-				}
-
-				if (MovieSceneToolHelpers::OverlapsSection(Track, Section, OutSections))
-				{
-					int32 AvailableRowIndex = MovieSceneToolHelpers::FindAvailableRowIndex(Track, Section, OutSections);
-					MaxAvailableRowIndex = FMath::Max(AvailableRowIndex, MaxAvailableRowIndex);
-				}
-			}
-
-			if (MaxAvailableRowIndex != -1)
-			{
-				for (UMovieSceneSection* Section : OutSections)
-				{
-					Section->SetRowIndex(MaxAvailableRowIndex);
-				}
-			}
-		}
+		Track->FixRowIndices();
 
 		// Remove sections that were pasted so that they aren't pasted again to another track
 		for (UMovieSceneSection* OutSection : OutSections)
diff --git a/Engine/Source/Editor/Sequencer/Private/Tools/EditToolDragOperations.cpp b/Engine/Source/Editor/Sequencer/Private/Tools/EditToolDragOperations.cpp
--- a/Engine/Source/Editor/Sequencer/Private/Tools/EditToolDragOperations.cpp
+++ b/Engine/Source/Editor/Sequencer/Private/Tools/EditToolDragOperations.cpp
@@ -930,11 +930,12 @@ void FMoveKeysAndSections::OnDrag(const FPointerEvent& MouseEvent, FVector2D Loc
 	// Get a list of the unique tracks in this selection and update their easing so previews draw interactively as you drag.
 	TSet<UMovieSceneTrack*> Tracks;
 	FMovieSceneSectionMovedParams SectionMovedParams(EPropertyChangeType::Interactive);
+	EMovieSceneSectionMovedResult SectionMovedResult = EMovieSceneSectionMovedResult::None;
 	for (UMovieSceneSection* Section : Sections)
 	{
 		if (UMovieSceneTrack* Track = Section->GetTypedOuter<UMovieSceneTrack>())
 		{
-			Track->OnSectionMoved(*Section, SectionMovedParams);
+			SectionMovedResult |= Track->OnSectionMoved(*Section, SectionMovedParams);
 			Tracks.Add(Track);
 		}
 	}
@@ -945,6 +946,10 @@ void FMoveKeysAndSections::OnDrag(const FPointerEvent& MouseEvent, FVector2D Loc
 	}
 
 	// If we changed the layout by rearranging sections we need to tell the Sequencer to rebuild things, otherwise just re-evaluate existing tracks.
+	if (EnumHasAnyFlags(SectionMovedResult, EMovieSceneSectionMovedResult::SectionsChanged))
+	{
+		bSectionMovementModifiedStructure = true;
+	}
 	if (bSectionMovementModifiedStructure)
 	{
 		Sequencer.NotifyMovieSceneDataChanged(EMovieSceneDataChangeType::MovieSceneStructureItemsChanged);
