Sequencer: Fix bug in anim mixer menu- if stitch wasn't enabled, we weren't getting any animation options to add.

#Synced-CL 51336034

diff --git a/Engine/Plugins/Experimental/MovieSceneAnimMixer/Source/MovieSceneAnimMixerEditor/Private/MovieSceneAnimationMixerTrackEditor.cpp b/Engine/Plugins/Experimental/MovieSceneAnimMixer/Source/MovieSceneAnimMixerEditor/Private/MovieSceneAnimationMixerTrackEditor.cpp
--- a/Engine/Plugins/Experimental/MovieSceneAnimMixer/Source/MovieSceneAnimMixerEditor/Private/MovieSceneAnimationMixerTrackEditor.cpp
+++ b/Engine/Plugins/Experimental/MovieSceneAnimMixer/Source/MovieSceneAnimMixerEditor/Private/MovieSceneAnimationMixerTrackEditor.cpp
@@ -997,7 +997,7 @@ TSharedRef<SWidget> FAnimationMixerTrackEditor::BuildAddSectionSubMenu(TWeakView
 	}
 
 	// Animation items
-	if (AllItemClasses.Num() > 0)
+	if (AllItemClasses.Num() > 0 || Providers.Num() > 0)
 	{
 		MenuBuilder.BeginSection(NAME_None, LOCTEXT("AnimationCategoryLabel", "Animation:"));
 
