Sequencer: Add support for blending 'child of root' to the mixer root motion settings.

This creates parity with the previous feature bBlendChildOfRoot on skeletal animation tracks to allow motion on the child of a root bone to be used to drive root motion instead.

There is now a 'Source' property in the Root Motion Settings decoration that can be used to pick between the root bone and the first animated child of the root.
There's also a fix here for quat->euler round trips creating gimbal lock problems in bone matching.

As part of this changelist, I removed the root motion extraction from the skeletal animation evaluation task in favor of it being handled generically for any section/track that has the root motion settings decoration by a separate FAnimNextExtractRootMotionTask- this allows us to have the bChildOfRoot behavior for any type of section not just skel anim sections. However, I did need to keep the old path in the skeletal animation eval task for 'lock root bone' cases which require direct access to the anim sequence.

#Synced-CL 51336034

diff --git a/Engine/Plugins/Experimental/MovieSceneAnimMixer/Source/MovieSceneAnimMixer/Private/AnimMixerBoneMatching.cpp b/Engine/Plugins/Experimental/MovieSceneAnimMixer/Source/MovieSceneAnimMixer/Private/AnimMixerBoneMatching.cpp
--- a/Engine/Plugins/Experimental/MovieSceneAnimMixer/Source/MovieSceneAnimMixer/Private/AnimMixerBoneMatching.cpp
+++ b/Engine/Plugins/Experimental/MovieSceneAnimMixer/Source/MovieSceneAnimMixer/Private/AnimMixerBoneMatching.cpp
@@ -44,33 +44,27 @@ static FTransform ComputeMatchOffset(
 		FQuat SecBoneQuat = SectionBoneCS.GetRotation();
 		UnderlyingBoneQuat.EnforceShortestArcWith(SecBoneQuat);
 
-		FRotator UnderlyingBoneRot(FRotator::MakeFromEuler(EulerFromQuat(UnderlyingBoneQuat, RotationOrder)));
-		FRotator SecBoneRot(FRotator::MakeFromEuler(EulerFromQuat(SecBoneQuat, RotationOrder)));
-		SecBoneRot.SetClosestToMe(UnderlyingBoneRot);
+		// Decompose the delta between the two quaternions to Euler (the delta is small, avoiding gimbal lock),
+		// zero unmatched axes, and recompose.
+		FQuat FullDelta = SecBoneQuat.Inverse() * UnderlyingBoneQuat;
+		FullDelta.Normalize();
+		FVector DeltaEuler = EulerFromQuat(FullDelta, RotationOrder);
 
 		if (!Settings.bMatchRotationZ)
 		{
-			UnderlyingBoneRot.Yaw = SecBoneRot.Yaw;
+			DeltaEuler.Z = 0.0;
 		}
 		if (!Settings.bMatchRotationY)
 		{
-			UnderlyingBoneRot.Pitch = SecBoneRot.Pitch;
+			DeltaEuler.Y = 0.0;
 		}
 		if (!Settings.bMatchRotationX)
 		{
-			UnderlyingBoneRot.Roll = SecBoneRot.Roll;
+			DeltaEuler.X = 0.0;
 		}
 
-		FQuat AdjustedUnderlyingBoneQuat = QuatFromEuler(UnderlyingBoneRot.Euler(), RotationOrder);
-		DeltaRot = SecBoneQuat.Inverse() * AdjustedUnderlyingBoneQuat;
-
-		// The Euler round-trip can produce -Q instead of Q (quaternion double
-		// cover). Both represent the same rotation, but negative W causes
-		// FTransform to interpret it as a 180° flip. Enforce positive W.
-		if (DeltaRot.W < 0.0)
-		{
-			DeltaRot = -DeltaRot;
-		}
+		DeltaRot = QuatFromEuler(DeltaEuler, RotationOrder);
+		DeltaRot.Normalize();
 	}
 
 	// Predict where the bone would be after applying just the rotation offset.
diff --git a/Engine/Plugins/Experimental/MovieSceneAnimMixer/Source/MovieSceneAnimMixer/Private/MovieSceneAnimationMixerTrack.cpp b/Engine/Plugins/Experimental/MovieSceneAnimMixer/Source/MovieSceneAnimMixer/Private/MovieSceneAnimationMixerTrack.cpp
--- a/Engine/Plugins/Experimental/MovieSceneAnimMixer/Source/MovieSceneAnimMixer/Private/MovieSceneAnimationMixerTrack.cpp
+++ b/Engine/Plugins/Experimental/MovieSceneAnimMixer/Source/MovieSceneAnimMixer/Private/MovieSceneAnimationMixerTrack.cpp
@@ -47,6 +47,9 @@ UMovieSceneRootMotionSettingsDecoration::UMovieSceneRootMotionSettingsDecoration
 	RootMotionSpace.SetEnum(StaticEnum<EMovieSceneRootMotionSpace>());
 	RootMotionSpace.SetDefault((uint8)EMovieSceneRootMotionSpace::AnimationSpace);
 
+	RootMotionSourceChannel.SetEnum(StaticEnum<EMovieSceneRootMotionSource>());
+	RootMotionSourceChannel.SetDefault((uint8)EMovieSceneRootMotionSource::RootBone);
+
 	TransformMode.SetEnum(StaticEnum<EMovieSceneRootMotionTransformMode>());
 	TransformMode.SetDefault((uint8)EMovieSceneRootMotionTransformMode::Offset);
 
@@ -70,6 +73,11 @@ EMovieSceneRootMotionTransformMode UMovieSceneRootMotionSettingsDecoration::GetR
 	return (EMovieSceneRootMotionTransformMode)TransformMode.GetDefault().Get((uint8)EMovieSceneRootMotionTransformMode::Offset);
 }
 
+EMovieSceneRootMotionSource UMovieSceneRootMotionSettingsDecoration::GetRootMotionSource() const
+{
+	return (EMovieSceneRootMotionSource)RootMotionSourceChannel.GetDefault().Get((uint8)EMovieSceneRootMotionSource::RootBone);
+}
+
 EMovieSceneChannelProxyType UMovieSceneRootMotionSettingsDecoration::PopulateChannelProxy(FMovieSceneChannelProxyData& OutProxyData)
 {
 	EMovieSceneRootMotionTransformMode CurrentTransformMode = GetRootTransformMode();
@@ -84,6 +92,13 @@ EMovieSceneChannelProxyType UMovieSceneRootMotionSettingsDecoration::PopulateCha
 		ChannelMetaData.SortOrder = SortOrder++;
 		OutProxyData.Add(RootMotionSpace, ChannelMetaData, TMovieSceneExternalValue<uint8>());
 	}
+	{
+		FMovieSceneChannelMetaData ChannelMetaData;
+		ChannelMetaData.SetIdentifiers("RootMotionSource", LOCTEXT("Source", "Source"), Group);
+		ChannelMetaData.WeakOwningObject = this;
+		ChannelMetaData.SortOrder = SortOrder++;
+		OutProxyData.Add(RootMotionSourceChannel, ChannelMetaData, TMovieSceneExternalValue<uint8>());
+	}
 	{
 		FMovieSceneChannelMetaData ChannelMetaData;
 		ChannelMetaData.SetIdentifiers("Root Mode", LOCTEXT("Mode", "Mode"), Group);
@@ -180,6 +195,7 @@ EMovieSceneChannelProxyType UMovieSceneRootMotionSettingsDecoration::PopulateCha
 	if (CurrentTransformMode != EMovieSceneRootMotionTransformMode::Asset)
 	{
 		OutProxyData.Add(RootMotionSpace);
+		OutProxyData.Add(RootMotionSourceChannel);
 		OutProxyData.Add(TransformMode);
 	}
 
@@ -208,8 +224,9 @@ void UMovieSceneRootMotionSettingsDecoration::ExtendEntityImpl(UMovieSceneEntity
 
 	FMovieSceneRootMotionSettings RootMotionSettings;
 	{
-		RootMotionSettings.RootMotionSpace = (EMovieSceneRootMotionSpace)RootMotionSpace.GetDefault().Get((uint8)EMovieSceneRootMotionSpace::AnimationSpace);
-		RootMotionSettings.TransformMode   = GetRootTransformMode();
+		RootMotionSettings.RootMotionSpace  = (EMovieSceneRootMotionSpace)RootMotionSpace.GetDefault().Get((uint8)EMovieSceneRootMotionSpace::AnimationSpace);
+		RootMotionSettings.RootMotionSource = GetRootMotionSource();
+		RootMotionSettings.TransformMode    = GetRootTransformMode();
 
 		UMovieSceneSkeletalAnimationSection* AnimSection = GetTypedOuter<UMovieSceneSkeletalAnimationSection>();
 		if (AnimSection)
diff --git a/Engine/Plugins/Experimental/MovieSceneAnimMixer/Source/MovieSceneAnimMixer/Private/Systems/MovieSceneAnimMixerSystem.cpp b/Engine/Plugins/Experimental/MovieSceneAnimMixer/Source/MovieSceneAnimMixer/Private/Systems/MovieSceneAnimMixerSystem.cpp
--- a/Engine/Plugins/Experimental/MovieSceneAnimMixer/Source/MovieSceneAnimMixer/Private/Systems/MovieSceneAnimMixerSystem.cpp
+++ b/Engine/Plugins/Experimental/MovieSceneAnimMixer/Source/MovieSceneAnimMixer/Private/Systems/MovieSceneAnimMixerSystem.cpp
@@ -346,13 +346,14 @@ static TSharedPtr<UE::UAF::FEvaluationProgram> BuildProgramFromEntries(
 			}
 
 			bool bNeedsMirror = (WrapEntry->MirrorTable != nullptr);
+			bool bNeedsExtraction = WrapEntry->RootMotionSettings.IsSet();
 			TOptional<FAnimNextConvertRootMotionToWorldSpaceTask> SpaceTask;
 			if (!Filter.bSkipRootMotionConversion)
 			{
 				SpaceTask = BuildSpaceConversionTask(*WrapEntry, Context, bRootComponentHasKeyedTransform);
 			}
 
-			if (!bNeedsMirror && !SpaceTask.IsSet())
+			if (!bNeedsMirror && !bNeedsExtraction && !SpaceTask.IsSet())
 			{
 				return WrapEntry->EvalTask;
 			}
@@ -367,6 +368,12 @@ static TSharedPtr<UE::UAF::FEvaluationProgram> BuildProgramFromEntries(
 				MiniProgram->AppendTaskPtr(MirrorTask);
 			}
 
+			if (bNeedsExtraction)
+			{
+				MiniProgram->AppendTask(FAnimNextExtractRootMotionTask::Make(
+					WrapEntry->RootMotionSettings->RootMotionSource));
+			}
+
 			if (Filter.CaptureAnimSpaceRoot && SpaceTask.IsSet())
 			{
 				SpaceTask->CaptureAnimSpaceRoot = Filter.CaptureAnimSpaceRoot;
@@ -528,6 +535,15 @@ static TSharedPtr<UE::UAF::FEvaluationProgram> BuildProgramFromEntries(
 				Program->AppendTaskPtr(MirrorTask);
 			}
 
+			// Extract root motion from the pose into a RootTransformAttribute.
+			// This runs after eval and mirror, before space conversion.
+			// Skipped for transitions (already handled in mini-programs).
+			if (MixerEntry->RootMotionSettings.IsSet() && !MixerEntry->bIsTransition)
+			{
+				Program->AppendTask(FAnimNextExtractRootMotionTask::Make(
+					MixerEntry->RootMotionSettings->RootMotionSource));
+			}
+
 			// Handle root space conversions. Skipped for bake (animation-space
 			// results) and for transitions (from/to are already converted to
 			// world space inside the wrapped mini-programs).
diff --git a/Engine/Plugins/Experimental/MovieSceneAnimMixer/Source/MovieSceneAnimMixer/Private/Systems/MovieSceneMixedSkeletalAnimationSystem.cpp b/Engine/Plugins/Experimental/MovieSceneAnimMixer/Source/MovieSceneAnimMixer/Private/Systems/MovieSceneMixedSkeletalAnimationSystem.cpp
--- a/Engine/Plugins/Experimental/MovieSceneAnimMixer/Source/MovieSceneAnimMixer/Private/Systems/MovieSceneMixedSkeletalAnimationSystem.cpp
+++ b/Engine/Plugins/Experimental/MovieSceneAnimMixer/Source/MovieSceneAnimMixer/Private/Systems/MovieSceneMixedSkeletalAnimationSystem.cpp
@@ -430,15 +430,16 @@ void FMovieSceneSkeletalAnimationEvaluationTask::Execute(UE::UAF::FEvaluationVM&
 		FKeyframeState Keyframe = VM.MakeUninitializedKeyframe(bIsAdditive);
 
 		const bool bUseRawData = FDecompressionTools::ShouldUseRawData(AnimSequencePtr, ExtractionContext);
-		const bool bUseRootMotion = AnimSequencePtr->bEnableRootMotion;
-
-		TOptional<FTransform> RootBoneTransform;
-		const int32 RootIndex = Keyframe.Pose.GetRefPose().GetLODBoneIndexFromSkeletonBoneIndex(0);
 
 		if (EnumHasAnyFlags(VM.GetFlags(), EEvaluationFlags::Bones))
 		{
 			FDecompressionTools::GetAnimationPose(AnimSequencePtr, ExtractionContext, Keyframe.Pose, bUseRawData);
 
+			// Legacy root motion override: the common animation track pre-computes root motion
+			// transforms and passes them in via RootMotionOverride. Apply the override to the
+			// appropriate bone and store it as the RootTransformAttribute directly. This sets
+			// the attribute so FAnimNextExtractRootMotionTask (which runs later) will skip
+			// re-extraction.
 			if (AnimationData.RootMotionOverride.IsSet())
 			{
 				int32 PoseIndex = INDEX_NONE;
@@ -454,21 +455,43 @@ void FMovieSceneSkeletalAnimationEvaluationTask::Execute(UE::UAF::FEvaluationVM&
 				{
 					Keyframe.Pose.LocalTransformsView[PoseIndex] = AnimationData.RootMotionOverride->RootMotion;
 				}
-			}
 
-			// Extract the root motion if we've chosen that option
-			if (bUseRootMotion)
-			{
-#if WITH_EDITOR
-				// If we've locked roots, the root won't hold any transform. We'll extract it later below.
-				if (!AnimSequencePtr->bForceRootLock)
+				FTransformAnimationAttribute* RootMotionAttribute = Keyframe.Attributes.FindOrAdd<FTransformAnimationAttribute>(AnimMixerComponents->RootTransformAttributeId);
+				FFloatAnimationAttribute* RootMotionWeight = Keyframe.Attributes.FindOrAdd<FFloatAnimationAttribute>(AnimMixerComponents->RootTransformWeightAttributeId);
+				RootMotionAttribute->Value = AnimationData.RootMotionOverride->RootMotion;
+				RootMotionWeight->Value = 1.f;
+
+				const int32 RootIndex = Keyframe.Pose.GetRefPose().GetLODBoneIndexFromSkeletonBoneIndex(0);
+				if (RootIndex != INDEX_NONE)
 				{
-#endif
-					RootBoneTransform = Keyframe.Pose.LocalTransformsView[RootIndex];
+					Keyframe.Pose.LocalTransformsView[RootIndex] = FTransform::Identity;
+				}
+			}
 #if WITH_EDITOR
+			// When bForceRootLock is true, the decompressor locks bone 0 to the
+			// reference pose so the actual root motion is not in the pose data.
+			// Pre-extract it from the raw track data so FAnimNextExtractRootMotionTask
+			// doesn't need access to the UAnimSequence. This only affects bone 0 as
+			// child bones are not locked, so FirstAnimatedChild extraction from the
+			// pose still works correctly.
+			else if (AnimSequencePtr->bForceRootLock && AnimSequencePtr->bEnableRootMotion)
+			{
+				FAnimExtractContext ExtractContext(AnimationData.ToPosition);
+				ExtractContext.bIgnoreRootLock = true;
+				FTransform RootTrackTransform = AnimSequencePtr->ExtractRootTrackTransform(ExtractContext, nullptr);
+
+				FTransformAnimationAttribute* RootMotionAttribute = Keyframe.Attributes.FindOrAdd<FTransformAnimationAttribute>(AnimMixerComponents->RootTransformAttributeId);
+				FFloatAnimationAttribute* RootMotionWeight = Keyframe.Attributes.FindOrAdd<FFloatAnimationAttribute>(AnimMixerComponents->RootTransformWeightAttributeId);
+				RootMotionAttribute->Value = RootTrackTransform;
+				RootMotionWeight->Value = 1.f;
+
+				const int32 RootIndex = Keyframe.Pose.GetRefPose().GetLODBoneIndexFromSkeletonBoneIndex(0);
+				if (RootIndex != INDEX_NONE)
+				{
+					Keyframe.Pose.LocalTransformsView[RootIndex] = FTransform::Identity;
 				}
-#endif
 			}
+#endif
 		}
 
 		if (EnumHasAnyFlags(VM.GetFlags(), EEvaluationFlags::Curves))
@@ -481,44 +504,6 @@ void FMovieSceneSkeletalAnimationEvaluationTask::Execute(UE::UAF::FEvaluationVM&
 			FDecompressionTools::GetAnimationAttributes(AnimSequencePtr, ExtractionContext, Keyframe.Pose.GetRefPose(), Keyframe.Attributes, bUseRawData);
 		}
 
-		// Trajectory is currently held as an attribute
-		if (EnumHasAnyFlags(VM.GetFlags(), EEvaluationFlags::Attributes | EEvaluationFlags::Trajectory))
-		{
-			// We intentionally do not use IAnimRootMotionProvider here to sample root motion:
-			//     We don't want to emit and blend RootMotionDelta attributes for every pose in the mixer,
-			//     we want to blend absolute RootMotionTransform attributes, and emit the final RootMotionDelta
-			//     attribute once, after all blending has been performed. This allows us to blend seamlessly with
-			//     world-space transforms from gamplay
-
-			// If we haven't already grabbed the root transform, grab it here in case we're not doing the full bone sample
-			if (!RootBoneTransform.IsSet())
-			{
-				if (AnimationData.RootMotionOverride.IsSet())
-				{
-					RootBoneTransform = AnimationData.RootMotionOverride->RootMotion;
-				}
-				else if (bUseRootMotion)
-				{
-					FAnimExtractContext ExtractContext(AnimationData.ToPosition);
-#if WITH_EDITOR
-					ExtractContext.bIgnoreRootLock = true;
-#endif
-					RootBoneTransform = AnimSequencePtr->ExtractRootTrackTransform(ExtractContext, nullptr);
-				}
-			}
-		}
-
-		if (RootBoneTransform.IsSet())
-		{
-			Keyframe.Pose.LocalTransformsView[RootIndex] = FTransform::Identity;
-
-			FTransformAnimationAttribute* RootMotionAttribute        = Keyframe.Attributes.FindOrAdd<FTransformAnimationAttribute>(AnimMixerComponents->RootTransformAttributeId);
-			FFloatAnimationAttribute*     RootMotionWeight           = Keyframe.Attributes.FindOrAdd<FFloatAnimationAttribute>(AnimMixerComponents->RootTransformWeightAttributeId);
-
-			RootMotionAttribute->Value = RootBoneTransform.GetValue();
-			RootMotionWeight->Value = 1.f;
-		}
-
 		VM.PushValue(KEYFRAME_STACK_NAME, MakeUnique<FKeyframeState>(MoveTemp(Keyframe)));
 	}
 	else
diff --git a/Engine/Plugins/Experimental/MovieSceneAnimMixer/Source/MovieSceneAnimMixer/Private/Systems/MovieSceneRootMotionSystem.cpp b/Engine/Plugins/Experimental/MovieSceneAnimMixer/Source/MovieSceneAnimMixer/Private/Systems/MovieSceneRootMotionSystem.cpp
--- a/Engine/Plugins/Experimental/MovieSceneAnimMixer/Source/MovieSceneAnimMixer/Private/Systems/MovieSceneRootMotionSystem.cpp
+++ b/Engine/Plugins/Experimental/MovieSceneAnimMixer/Source/MovieSceneAnimMixer/Private/Systems/MovieSceneRootMotionSystem.cpp
@@ -576,6 +576,106 @@ void FAnimNextSetRootTransformAttributeTask::Execute(UE::UAF::FEvaluationVM& VM)
 	RootMotionAttribute->Value = Transform;
 }
 
+FAnimNextExtractRootMotionTask FAnimNextExtractRootMotionTask::Make(EMovieSceneRootMotionSource InSource)
+{
+	FAnimNextExtractRootMotionTask Task;
+	Task.Source = InSource;
+	return Task;
+}
+
+void FAnimNextExtractRootMotionTask::Execute(UE::UAF::FEvaluationVM& VM) const
+{
+	using namespace UE::UAF;
+	using namespace UE::MovieScene;
+
+	if (VM.GetActiveNamedSet())
+	{
+		return;
+	}
+
+	if (!EnumHasAnyFlags(VM.GetFlags(), EEvaluationFlags::Bones | EEvaluationFlags::Attributes | EEvaluationFlags::Trajectory))
+	{
+		return;
+	}
+
+	TUniquePtr<FKeyframeState>* Keyframe = VM.PeekValueMutable<TUniquePtr<FKeyframeState>>(KEYFRAME_STACK_NAME, 0);
+	if (!Keyframe || !Keyframe->IsValid())
+	{
+		return;
+	}
+
+	const FAnimMixerComponentTypes* AnimMixerComponents = FAnimMixerComponentTypes::Get();
+
+	// If a RootTransformAttribute already exists (e.g. from a legacy RootMotionOverride),
+	// don't overwrite it.
+	if ((*Keyframe)->Attributes.Find<FTransformAnimationAttribute>(AnimMixerComponents->RootTransformAttributeId))
+	{
+		return;
+	}
+
+	const int32 RootLODIndex = (*Keyframe)->Pose.GetRefPose().GetLODBoneIndexFromSkeletonBoneIndex(0);
+	if (RootLODIndex == INDEX_NONE)
+	{
+		return;
+	}
+
+	const int32 NumBones = (*Keyframe)->Pose.LocalTransformsView.Num();
+	int32 SourceLODIndex = RootLODIndex;
+
+	if (Source == EMovieSceneRootMotionSource::FirstAnimatedChild)
+	{
+		// Find the first direct child of root that has a non-zero-translation transform
+		const TArrayView<const FBoneIndexType> ParentMap = (*Keyframe)->Pose.GetLODBoneIndexToParentLODBoneIndexMap();
+
+		for (int32 BoneIdx = 0; BoneIdx < NumBones; ++BoneIdx)
+		{
+			if (ParentMap.IsValidIndex(BoneIdx)
+				&& ParentMap[BoneIdx] == RootLODIndex)
+			{
+				FTransform BoneTransform = (*Keyframe)->Pose.LocalTransformsView[BoneIdx];
+				if (!BoneTransform.GetTranslation().IsNearlyZero())
+				{
+					SourceLODIndex = BoneIdx;
+					break;
+				}
+			}
+		}
+	}
+
+	if (SourceLODIndex < 0 || SourceLODIndex >= NumBones)
+	{
+		return;
+	}
+
+	const auto& RefPose = (*Keyframe)->Pose.GetRefPose();
+	FTransform SourceBoneTransform = (*Keyframe)->Pose.LocalTransformsView[SourceLODIndex];
+
+	// When extracting from a child bone, use the delta from the rest pose
+	// so the root motion only captures locomotion, not the skeletal offset
+	// (e.g., pelvis height). Restore the ref pose on the child bone so
+	// the skeleton retains its correct shape.
+	FTransform RootMotionTransform;
+	if (SourceLODIndex != RootLODIndex)
+	{
+		FTransform RefTransform = RefPose.ReferenceLocalTransforms[SourceLODIndex];
+		RootMotionTransform = RefTransform.Inverse() * SourceBoneTransform;
+		(*Keyframe)->Pose.LocalTransformsView[SourceLODIndex] = RefTransform;
+	}
+	else
+	{
+		RootMotionTransform = SourceBoneTransform;
+		(*Keyframe)->Pose.LocalTransformsView[SourceLODIndex] = FTransform::Identity;
+	}
+
+	FTransformAnimationAttribute* RootMotionAttribute =
+		(*Keyframe)->Attributes.FindOrAdd<FTransformAnimationAttribute>(AnimMixerComponents->RootTransformAttributeId);
+	FFloatAnimationAttribute* RootMotionWeight =
+		(*Keyframe)->Attributes.FindOrAdd<FFloatAnimationAttribute>(AnimMixerComponents->RootTransformWeightAttributeId);
+
+	RootMotionAttribute->Value = RootMotionTransform;
+	RootMotionWeight->Value = 1.f;
+}
+
 UMovieSceneRootMotionSystem::UMovieSceneRootMotionSystem(const FObjectInitializer& ObjInit)
 	: Super(ObjInit)
 {
diff --git a/Engine/Plugins/Experimental/MovieSceneAnimMixer/Source/MovieSceneAnimMixer/Public/MovieSceneAnimationMixerTrack.h b/Engine/Plugins/Experimental/MovieSceneAnimMixer/Source/MovieSceneAnimMixer/Public/MovieSceneAnimationMixerTrack.h
--- a/Engine/Plugins/Experimental/MovieSceneAnimMixer/Source/MovieSceneAnimMixer/Public/MovieSceneAnimationMixerTrack.h
+++ b/Engine/Plugins/Experimental/MovieSceneAnimMixer/Source/MovieSceneAnimMixer/Public/MovieSceneAnimationMixerTrack.h
@@ -103,6 +103,7 @@ public:
 	UMovieSceneRootMotionSettingsDecoration(const FObjectInitializer& ObjInit);
 
 	MOVIESCENEANIMMIXER_API EMovieSceneRootMotionTransformMode GetRootTransformMode() const;
+	MOVIESCENEANIMMIXER_API EMovieSceneRootMotionSource GetRootMotionSource() const;
 
 	virtual EMovieSceneChannelProxyType PopulateChannelProxy(FMovieSceneChannelProxyData& OutProxyData) override;
 	virtual void ExtendEntityImpl(UMovieSceneEntitySystemLinker* EntityLinker, const FEntityImportParams& Params, FImportedEntity* OutImportedEntity) override;
@@ -121,6 +122,9 @@ public:
 	UPROPERTY()
 	FMovieSceneByteChannelDefaultOnly RootMotionSpace;
 
+	UPROPERTY()
+	FMovieSceneByteChannelDefaultOnly RootMotionSourceChannel;
+
 	UPROPERTY()
 	FMovieSceneByteChannelDefaultOnly TransformMode;
 
@@ -170,7 +174,7 @@ MOVIESCENEANIMMIXER_API bool WouldCreateBoneMatchCycle(UMovieSceneSection* Candi
 /**
  *
  */
-UCLASS(MinimalAPI)
+UCLASS(MinimalAPI, HideCategories=("RootMotions"))
 class UMovieSceneAnimationMixerTrack
 	: public UMovieSceneCommonAnimationTrack
 {
diff --git a/Engine/Plugins/Experimental/MovieSceneAnimMixer/Source/MovieSceneAnimMixer/Public/Systems/MovieSceneRootMotionSystem.h b/Engine/Plugins/Experimental/MovieSceneAnimMixer/Source/MovieSceneAnimMixer/Public/Systems/MovieSceneRootMotionSystem.h
--- a/Engine/Plugins/Experimental/MovieSceneAnimMixer/Source/MovieSceneAnimMixer/Public/Systems/MovieSceneRootMotionSystem.h
+++ b/Engine/Plugins/Experimental/MovieSceneAnimMixer/Source/MovieSceneAnimMixer/Public/Systems/MovieSceneRootMotionSystem.h
@@ -16,6 +16,17 @@ class UMovieSceneAnimMixerSystem;
 class UMovieSceneRootMotionSection;
 class UMovieSceneRootMotionSystem;
 
+// Which bone to extract root motion from in the pose
+UENUM()
+enum class EMovieSceneRootMotionSource : uint8
+{
+	// Extract root motion from skeleton bone 0 (standard root motion)
+	RootBone,
+	// Auto-detect the first child of root that has animation data.
+	// Use when the root bone is static and locomotion lives on a child bone.
+	FirstAnimatedChild,
+};
+
 // Enum describing what space that root motion should be applied in.
 UENUM()
 enum class EMovieSceneRootMotionSpace : uint8
@@ -77,6 +88,7 @@ struct FMovieSceneRootMotionSettings
 	// What space to apply root motion in. Defaults to animation space.
 	EMovieSceneRootMotionSpace RootMotionSpace = EMovieSceneRootMotionSpace::AnimationSpace;
 	EMovieSceneRootMotionTransformMode TransformMode = EMovieSceneRootMotionTransformMode::Offset;
+	EMovieSceneRootMotionSource RootMotionSource = EMovieSceneRootMotionSource::RootBone;
 	ESwapRootBone LegacySwapRootBone = ESwapRootBone::SwapRootBone_None;
 
 	uint8 bHasRootMotionOverride : 1 = 0;
@@ -233,6 +245,30 @@ struct FAnimNextSetRootTransformAttributeTask : public FAnimNextEvaluationTask
 	FTransform Transform;
 };
 
+// Generic root motion extraction task. Reads the top keyframe's pose,
+// extracts root motion from either the root bone or the first animated
+// child of root based on the Source setting, and stores the transform as a
+// RootTransformAttribute. For the root bone, the bone is zeroed in the
+// pose. For a child bone, the delta from the reference pose is extracted
+// and the bone is restored to its reference transform.
+// Inserted into an entry mixer program after the section's eval task and
+// before the space conversion task.
+USTRUCT()
+struct FAnimNextExtractRootMotionTask : public FAnimNextEvaluationTask
+{
+	GENERATED_BODY()
+
+	DECLARE_ANIM_EVALUATION_TASK(FAnimNextExtractRootMotionTask)
+
+	FAnimNextExtractRootMotionTask() = default;
+
+	static FAnimNextExtractRootMotionTask Make(EMovieSceneRootMotionSource InSource);
+
+	virtual void Execute(UE::UAF::FEvaluationVM& VM) const override;
+
+	EMovieSceneRootMotionSource Source = EMovieSceneRootMotionSource::RootBone;
+};
+
 // Takes in evaluation tasks on mixers.
 // Mixes just the root motion attributes.
 // Converts it from animation space to either additive actor or component space (based on which attribute used).
