Sequencer: Fix anim mixer layer deactivation

The base class SetIsDeactivated() iterates ITrackExtension descendants
and calls SetEvalDisabled on the track from GetTrack(). For section-based
mixer layers, GetTrack() returns the parent mixer track, disabling the
entire mixer instead of just the target layer. Override IsDeactivated()
and SetIsDeactivated() in FAnimationMixerLayerModel to use per-row
deactivation on the parent track for section layers, and per-track
deactivation for child track layers.

Also override IsDeactivated/SetIsDeactivated on FAnimationMixerTrackModel
so that deactivating the mixer track itself only sets the track-level flag
without leaking into per-row flags on individual layers.

Also override IsDeactivated/SetIsDeactivated on FSectionOutlinerModel to
toggle the section's active flag directly, so deactivating individual
sections within expanded layers works.

#Synced-CL 51540725

diff --git a/Engine/Plugins/Experimental/MovieSceneAnimMixer/Source/MovieSceneAnimMixerEditor/Private/MVVM/AnimationMixerLayerModel.cpp b/Engine/Plugins/Experimental/MovieSceneAnimMixer/Source/MovieSceneAnimMixerEditor/Private/MVVM/AnimationMixerLayerModel.cpp
--- a/Engine/Plugins/Experimental/MovieSceneAnimMixer/Source/MovieSceneAnimMixerEditor/Private/MVVM/AnimationMixerLayerModel.cpp
+++ b/Engine/Plugins/Experimental/MovieSceneAnimMixer/Source/MovieSceneAnimMixerEditor/Private/MVVM/AnimationMixerLayerModel.cpp
@@ -1083,6 +1083,71 @@ void FAnimationMixerLayerModel::SetExpansion(bool bInIsExpanded)
 	RefreshLayout(true);
 }
 
+/*-----------------------------------------------------------------------------
+	FEvaluableOutlinerItemModel - Deactivation
+	The base class implementation iterates ITrackExtension descendants and
+	calls SetEvalDisabled on the track returned by GetTrack(). For mixer
+	layers that don't have a child track, GetTrack() returns the parent
+	mixer track, which would deactivate the entire mixer instead of just
+	this layer. Override to use per-row deactivation on the parent track.
+-----------------------------------------------------------------------------*/
+
+bool FAnimationMixerLayerModel::IsDeactivated() const
+{
+	// Child track layer: check if the child track itself is eval-disabled
+	if (UMovieSceneTrack* ChildTrack = WeakChildTrack.Get())
+	{
+		return ChildTrack->IsEvalDisabled(/*bInCheckLocal=*/false);
+	}
+
+	// Section layer: check row-level deactivation on the parent mixer track
+	if (UMovieSceneTrack* ParentTrack = WeakParentTrack.Get())
+	{
+		return ParentTrack->IsRowEvalDisabled(GetRowIndex(), /*bInCheckLocal=*/false);
+	}
+	return false;
+}
+
+void FAnimationMixerLayerModel::SetIsDeactivated(bool bInIsDeactivated)
+{
+	bool bChanged = false;
+
+	// Child track layer: disable the child track itself
+	if (UMovieSceneTrack* ChildTrack = WeakChildTrack.Get())
+	{
+		if (bInIsDeactivated != ChildTrack->IsEvalDisabled(/*bInCheckLocal=*/false))
+		{
+			ChildTrack->Modify();
+			ChildTrack->SetEvalDisabled(bInIsDeactivated);
+			bChanged = true;
+		}
+	}
+	else
+	{
+		// Section layer: use row-level deactivation on the parent mixer track
+		UMovieSceneTrack* ParentTrack = WeakParentTrack.Get();
+		if (ParentTrack && bInIsDeactivated != ParentTrack->IsRowEvalDisabled(GetRowIndex(), /*bInCheckLocal=*/false))
+		{
+			ParentTrack->Modify();
+			ParentTrack->SetRowEvalDisabled(bInIsDeactivated, GetRowIndex());
+			bChanged = true;
+		}
+	}
+
+	if (bChanged)
+	{
+		TSharedPtr<FSequencerEditorViewModel> EditorViewModel = GetEditor();
+		if (EditorViewModel)
+		{
+			TSharedPtr<ISequencer> Sequencer = EditorViewModel->GetSequencer();
+			if (Sequencer)
+			{
+				Sequencer->NotifyMovieSceneDataChanged(EMovieSceneDataChangeType::TrackValueChanged);
+			}
+		}
+	}
+}
+
 /*-----------------------------------------------------------------------------
 	IDeletableExtension
 -----------------------------------------------------------------------------*/
@@ -1217,7 +1282,6 @@ void FAnimationMixerLayerModel::SetIsSoloed(bool bIsSoloed)
 	const bool bAlwaysMarkDirty = false;
 	Track->Modify(bAlwaysMarkDirty);
 
-	// TODO: This doesn't work and is soloing the whole mixer?
 	if (UMovieSceneTrackRowDecoration* TrackRowDecoration = Cast<UMovieSceneTrackRowDecoration>(Track->GetOrCreateDecoration(UMovieSceneTrackRowDecoration::StaticClass())))
 	{
 		TrackRowDecoration->SetSoloed(GetRowIndex(), bIsSoloed);
diff --git a/Engine/Plugins/Experimental/MovieSceneAnimMixer/Source/MovieSceneAnimMixerEditor/Private/MVVM/AnimationMixerTrackModel.cpp b/Engine/Plugins/Experimental/MovieSceneAnimMixer/Source/MovieSceneAnimMixerEditor/Private/MVVM/AnimationMixerTrackModel.cpp
--- a/Engine/Plugins/Experimental/MovieSceneAnimMixer/Source/MovieSceneAnimMixerEditor/Private/MVVM/AnimationMixerTrackModel.cpp
+++ b/Engine/Plugins/Experimental/MovieSceneAnimMixer/Source/MovieSceneAnimMixerEditor/Private/MVVM/AnimationMixerTrackModel.cpp
@@ -21,7 +21,9 @@
 #include "MVVM/ViewModels/ViewModelIterators.h"
 #include "MVVM/SharedViewModelData.h"
 #include "MVVM/ViewModels/ViewModelHierarchy.h"
+#include "ISequencer.h"
 #include "ISequencerTrackEditor.h"
+#include "MVVM/ViewModels/SequencerEditorViewModel.h"
 #include "MovieSceneSection.h"
 #include "Styling/CoreStyle.h"
 #include "Styling/SlateTypes.h"
@@ -41,6 +43,33 @@ FAnimationMixerTrackModel::FAnimationMixerTrackModel(UMovieSceneTrack* Track)
 {
 }
 
+bool FAnimationMixerTrackModel::IsDeactivated() const
+{
+	UMovieSceneTrack* Track = GetTrack();
+	return Track && Track->IsEvalDisabled(/*bInCheckLocal=*/false);
+}
+
+void FAnimationMixerTrackModel::SetIsDeactivated(bool bInIsDeactivated)
+{
+	UMovieSceneTrack* Track = GetTrack();
+	if (!Track || bInIsDeactivated == Track->IsEvalDisabled(/*bInCheckLocal=*/false))
+	{
+		return;
+	}
+
+	Track->Modify();
+	Track->SetEvalDisabled(bInIsDeactivated);
+
+	TSharedPtr<FSequencerEditorViewModel> EditorViewModel = GetEditor();
+	if (EditorViewModel)
+	{
+		if (TSharedPtr<ISequencer> Sequencer = EditorViewModel->GetSequencer())
+		{
+			Sequencer->NotifyMovieSceneDataChanged(EMovieSceneDataChangeType::TrackValueChanged);
+		}
+	}
+}
+
 void FAnimationMixerTrackModel::ForceUpdate()
 {
 	UMovieSceneAnimationMixerTrack* MixerTrack = Cast<UMovieSceneAnimationMixerTrack>(GetTrack());
diff --git a/Engine/Plugins/Experimental/MovieSceneAnimMixer/Source/MovieSceneAnimMixerEditor/Public/MVVM/AnimationMixerLayerModel.h b/Engine/Plugins/Experimental/MovieSceneAnimMixer/Source/MovieSceneAnimMixerEditor/Public/MVVM/AnimationMixerLayerModel.h
--- a/Engine/Plugins/Experimental/MovieSceneAnimMixer/Source/MovieSceneAnimMixerEditor/Public/MVVM/AnimationMixerLayerModel.h
+++ b/Engine/Plugins/Experimental/MovieSceneAnimMixer/Source/MovieSceneAnimMixerEditor/Public/MVVM/AnimationMixerLayerModel.h
@@ -163,6 +163,10 @@ public:
 	/*~ IDraggableOutlinerExtension */
 	virtual bool CanDrag() const override;
 
+	/*~ FEvaluableOutlinerItemModel */
+	virtual bool IsDeactivated() const override;
+	virtual void SetIsDeactivated(bool bInIsDeactivated) override;
+
 	/*~ FOutlinerItemModel */
 	virtual bool HasCurves() const override;
 	virtual void BuildContextMenu(FMenuBuilder& MenuBuilder) override;
diff --git a/Engine/Plugins/Experimental/MovieSceneAnimMixer/Source/MovieSceneAnimMixerEditor/Public/MVVM/AnimationMixerTrackModel.h b/Engine/Plugins/Experimental/MovieSceneAnimMixer/Source/MovieSceneAnimMixerEditor/Public/MVVM/AnimationMixerTrackModel.h
--- a/Engine/Plugins/Experimental/MovieSceneAnimMixer/Source/MovieSceneAnimMixerEditor/Public/MVVM/AnimationMixerTrackModel.h
+++ b/Engine/Plugins/Experimental/MovieSceneAnimMixer/Source/MovieSceneAnimMixerEditor/Public/MVVM/AnimationMixerTrackModel.h
@@ -54,6 +54,10 @@ public:
 
 	explicit FAnimationMixerTrackModel(UMovieSceneTrack* Track);
 
+	/*~ FEvaluableOutlinerItemModel */
+	virtual bool IsDeactivated() const override;
+	virtual void SetIsDeactivated(bool bInIsDeactivated) override;
+
 	/*~ ISectionDragOwnerExtension - Preview-based drag with outliner-style drop indicators */
 	virtual void OnBeginSectionVerticalDrag() override;
 	virtual bool OnSectionVerticalDrag(const FSectionVerticalDragContext& Context) override;
diff --git a/Engine/Source/Editor/Sequencer/Private/MVVM/ViewModels/SectionOutlinerModel.cpp b/Engine/Source/Editor/Sequencer/Private/MVVM/ViewModels/SectionOutlinerModel.cpp
--- a/Engine/Source/Editor/Sequencer/Private/MVVM/ViewModels/SectionOutlinerModel.cpp
+++ b/Engine/Source/Editor/Sequencer/Private/MVVM/ViewModels/SectionOutlinerModel.cpp
@@ -110,6 +110,37 @@ UMovieSceneTrack* FSectionOutlinerModel::GetTrack() const
 	return nullptr;
 }
 
+/*-----------------------------------------------------------------------------
+	FEvaluableOutlinerItemModel
+-----------------------------------------------------------------------------*/
+
+bool FSectionOutlinerModel::IsDeactivated() const
+{
+	UMovieSceneSection* Section = WeakSection.Get();
+	return Section && !Section->IsActive();
+}
+
+void FSectionOutlinerModel::SetIsDeactivated(bool bInIsDeactivated)
+{
+	UMovieSceneSection* Section = WeakSection.Get();
+	if (!Section || Section->IsActive() == !bInIsDeactivated)
+	{
+		return;
+	}
+
+	Section->Modify();
+	Section->SetIsActive(!bInIsDeactivated);
+
+	TSharedPtr<FSequencerEditorViewModel> EditorViewModel = GetEditor();
+	if (EditorViewModel)
+	{
+		if (TSharedPtr<ISequencer> Sequencer = EditorViewModel->GetSequencer())
+		{
+			Sequencer->NotifyMovieSceneDataChanged(EMovieSceneDataChangeType::TrackValueChanged);
+		}
+	}
+}
+
 /*-----------------------------------------------------------------------------
 	FOutlinerItemModel
 -----------------------------------------------------------------------------*/
diff --git a/Engine/Source/Editor/Sequencer/Public/MVVM/ViewModels/SectionOutlinerModel.h b/Engine/Source/Editor/Sequencer/Public/MVVM/ViewModels/SectionOutlinerModel.h
--- a/Engine/Source/Editor/Sequencer/Public/MVVM/ViewModels/SectionOutlinerModel.h
+++ b/Engine/Source/Editor/Sequencer/Public/MVVM/ViewModels/SectionOutlinerModel.h
@@ -75,6 +75,10 @@ public:
 
 	UE_API UMovieSceneTrack* GetTrack() const;
 
+	/*~ FEvaluableOutlinerItemModel */
+	UE_API bool IsDeactivated() const override;
+	UE_API void SetIsDeactivated(bool bInIsDeactivated) override;
+
 	/*~ FOutlinerItemModel */
 	UE_API void BuildContextMenu(FMenuBuilder& MenuBuilder) override;
 	UE_API void BuildSidebarMenu(FMenuBuilder& MenuBuilder) override;
