Sequencer: Fix Control Rig controls not exposed on add to anim mixer

When a child track is added to an existing layer, the recycled layer
model still thinks it is a section-type layer. Discard the recycled
model when the layer's child track state has changed so a fresh model
is created with the correct type.

#Synced-CL 51540725

diff --git a/Engine/Plugins/Experimental/MovieSceneAnimMixer/Source/MovieSceneAnimMixerEditor/Private/MVVM/AnimationMixerTrackModel.cpp b/Engine/Plugins/Experimental/MovieSceneAnimMixer/Source/MovieSceneAnimMixerEditor/Private/MVVM/AnimationMixerTrackModel.cpp
--- a/Engine/Plugins/Experimental/MovieSceneAnimMixer/Source/MovieSceneAnimMixerEditor/Private/MVVM/AnimationMixerTrackModel.cpp
+++ b/Engine/Plugins/Experimental/MovieSceneAnimMixer/Source/MovieSceneAnimMixerEditor/Private/MVVM/AnimationMixerTrackModel.cpp
@@ -111,12 +111,17 @@ void FAnimationMixerTrackModel::ForceUpdate()
 		{
 			if (ExistingModel && ExistingModel->GetLayer() == Layer)
 			{
+				// If the layer's child track state changed, discard the old model and create fresh
+				if (ExistingModel->HasChildTrack() != Layer->HasChildTrack())
+				{
+					break;
+				}
 				LayerModel = ExistingModel;
 				break;
 			}
 		}
 
-		// Create new layer model if not found
+		// Create new layer model if not found or if layer type changed
 		bool bNeedsLayout = false;
 		if (!LayerModel)
 		{
