Masking in Anim Mixer - Runtime

-Adding a mask decoration that will store sections for Hierarchy Table Assets to apply as masks
-These mask decorations can be added to animation mixer layers, which will apply the resulting animation pose of that layer per bone based on that mask
-Section provider now allows new sections to be added, and sections to be removed.
-Updated Mixer Track to populate the evaluation tree with the sectoins found in it's layers' decorations
-Updated Mixer Track to update easing on the sections found on layer decoraitons
-Anim mixer system now stores a blend task per mixer entry, which is used to apply masks, but could be used in the future to support other blend modes for layers.
-Masks sections support easing and will blend between each other on a per-bone level.

#Synced-CL 51628111

diff --git a/Engine/Plugins/Experimental/MovieSceneAnimMixer/Source/MovieSceneAnimMixer/Private/AnimMixerComponentTypes.cpp b/Engine/Plugins/Experimental/MovieSceneAnimMixer/Source/MovieSceneAnimMixer/Private/AnimMixerComponentTypes.cpp
--- a/Engine/Plugins/Experimental/MovieSceneAnimMixer/Source/MovieSceneAnimMixer/Private/AnimMixerComponentTypes.cpp
+++ b/Engine/Plugins/Experimental/MovieSceneAnimMixer/Source/MovieSceneAnimMixer/Private/AnimMixerComponentTypes.cpp
@@ -43,6 +43,7 @@ namespace UE::MovieScene
 		ComponentRegistry->NewComponentType(&Target, TEXT("Mixed Animation Target"));
 		ComponentRegistry->NewComponentType(&Task, TEXT("Mixed Animation Task"));
 		ComponentRegistry->NewComponentType(&MixerTask, TEXT("Mixed Animation Mixer Task"));
+		ComponentRegistry->NewComponentType(&BlendProviderData, TEXT("Mixed Animation Blend Provider Data"));
 		ComponentRegistry->NewComponentType(&RootMotionSettings, TEXT("Root Motion Settings"));
 		ComponentRegistry->NewComponentType(&MeshComponent, TEXT("Mixed Animation Mesh Component"));
 		ComponentRegistry->NewComponentType(&MixerRootMotion, TEXT("Root Motion"));
@@ -63,6 +64,7 @@ namespace UE::MovieScene
 		ComponentRegistry->Factories.DuplicateChildComponent(Target);
 		ComponentRegistry->Factories.DuplicateChildComponent(Task);
 		ComponentRegistry->Factories.DuplicateChildComponent(MixerTask);
+		ComponentRegistry->Factories.DuplicateChildComponent(BlendProviderData);
 		ComponentRegistry->Factories.DuplicateChildComponent(RootMotionSettings);
 		ComponentRegistry->Factories.DuplicateChildComponent(MeshComponent);
 		ComponentRegistry->Factories.DuplicateChildComponent(MixerRootMotion);
diff --git a/Engine/Plugins/Experimental/MovieSceneAnimMixer/Source/MovieSceneAnimMixer/Private/MovieSceneAnimMixerMaskSection.cpp b/Engine/Plugins/Experimental/MovieSceneAnimMixer/Source/MovieSceneAnimMixer/Private/MovieSceneAnimMixerMaskSection.cpp
new file mode 100644
--- /dev/null
+++ b/Engine/Plugins/Experimental/MovieSceneAnimMixer/Source/MovieSceneAnimMixer/Private/MovieSceneAnimMixerMaskSection.cpp
@@ -0,0 +1,49 @@
+﻿// Copyright Epic Games, Inc. All Rights Reserved.
+
+#include "MovieSceneAnimMixerMaskSection.h"
+
+#include "AnimMixerComponentTypes.h"
+#include "MovieSceneAnimationMaskDecoration.h"
+#include "MovieSceneAnimationMixerTrack.h"
+#include "EntitySystem/BuiltInComponentTypes.h"
+#include "MaskProfile/HierarchyTableTypeMask.h"
+
+#include UE_INLINE_GENERATED_CPP_BY_NAME(MovieSceneAnimMixerMaskSection)
+
+#define LOCTEXT_NAMESPACE "MovieSceneAnimationMixerMaskSection"
+
+UMovieSceneAnimMixerMaskSection::UMovieSceneAnimMixerMaskSection()
+{
+	BlendType = EMovieSceneBlendType::Absolute;
+	// This class implements, IMovieSceneAnimationMixerItemInterface, but since it's owned by a decoration, it is not added to the mixer track, which
+	// is where other mixer items have their tint color set. So the tint color is set here instead.
+	SetColorTint(MaskTintColor);
+}
+
+void UMovieSceneAnimMixerMaskSection::ImportEntityImpl(UMovieSceneEntitySystemLinker* EntityLinker, const FEntityImportParams& Params,
+	FImportedEntity* OutImportedEntity)
+{
+	using namespace UE::MovieScene;
+
+	const FAnimMixerComponentTypes* AnimMixerComponentTypes = FAnimMixerComponentTypes::Get();
+
+	BuildDefaultComponents(EntityLinker, Params, OutImportedEntity);
+	
+	if (!OwningLayer)
+	{
+		return;
+	}
+
+	FMovieSceneMaskBlendProviderData BlendData;
+
+	BlendData.BlendMask = BlendMask;
+
+
+	OutImportedEntity->AddBuilder(
+	FEntityBuilder()
+	.Add(AnimMixerComponentTypes->MixerLayer, OwningLayer)
+	.Add(AnimMixerComponentTypes->BlendProviderData, 	TInstancedStruct<FMovieSceneMaskBlendProviderData>::Make(MoveTemp(BlendData)))
+	.AddMutualComponents());
+}
+
+#undef LOCTEXT_NAMESPACE
diff --git a/Engine/Plugins/Experimental/MovieSceneAnimMixer/Source/MovieSceneAnimMixer/Private/MovieSceneAnimationMaskDecoration.cpp b/Engine/Plugins/Experimental/MovieSceneAnimMixer/Source/MovieSceneAnimMixer/Private/MovieSceneAnimationMaskDecoration.cpp
new file mode 100644
--- /dev/null
+++ b/Engine/Plugins/Experimental/MovieSceneAnimMixer/Source/MovieSceneAnimMixer/Private/MovieSceneAnimationMaskDecoration.cpp
@@ -0,0 +1,215 @@
+﻿// Copyright Epic Games, Inc. All Rights Reserved.
+
+#include "MovieSceneAnimationMaskDecoration.h"
+
+#include "MovieScene.h"
+#include "MovieSceneAnimMixerMaskSection.h"
+#include "EvaluationVM/Tasks/BlendKeyframesPerBone.h"
+
+TSharedPtr<FAnimNextEvaluationTask> FMovieSceneMaskBlendProvider::GetBlendTask()
+{
+	BlendTask = nullptr;
+	if (AccumulationBoneWeights.IsEmpty())
+	{
+		ResetAccumulationData();
+		return BlendTask;
+	}
+
+	// Curve data needs to be re-allocated due to the nature of TNamedValueArray.
+	// @TODO This may be optimized by checking if the data has changed, and not emptying the weights if it hasn't, which should be the case more often than not.
+	TaskCurveWeights.Empty();
+	TaskAttributeWeights.Reset();
+	TaskBoneWeights.Reset();
+
+	// Second pass of calculating bone weights. Previously the damping values were accumulated, but they need to be inverted to get the weight again.
+	for (int32 WeightIndex = 0; WeightIndex < AccumulationBoneWeights.Num(); ++WeightIndex)
+	{
+		// Final Weight = 1.0f - Accumulated Damping Values. 
+		TaskBoneWeights.Add(1.0f - AccumulationBoneWeights[WeightIndex]);
+	}
+
+	for (const TTuple<UE::Anim::FAttributeId, float> CurrentAttributeWeight : AttributeWeightsMap)
+	{
+		TaskAttributeWeights.Add(FBlendProfileStandaloneCachedData::FMaskedAttributeWeight(CurrentAttributeWeight.Key, 1.0 - CurrentAttributeWeight.Value));
+	}
+
+	for (const TTuple<FName, float> CurrentCurveWeight : CurveWeightsMap)
+	{
+		TaskCurveWeights.Add(UE::Anim::FCurveElement(CurrentCurveWeight.Key, 1.0 - CurrentCurveWeight.Value));
+	}
+
+	BlendTask = MakeShared<FAnimNextBlendKeyframePerBoneWithScaleTask>(
+	FAnimNextBlendKeyframePerBoneWithScaleTask::Make(MaskSkeleton.Get(), &TaskBoneWeights,
+	&TaskCurveWeights, &TaskAttributeWeights, 1.0f));
+
+	ResetAccumulationData();
+	return BlendTask;
+}
+
+void FMovieSceneMaskBlendProvider::RegisterTaskData(const TInstancedStruct<FMovieSceneAnimMixerBlendProviderData>& InData, const double Weight)
+{
+	if (const FMovieSceneMaskBlendProviderData* MaskData = InData.GetPtr<FMovieSceneMaskBlendProviderData>())
+	{
+		if (MaskData->BlendMask.IsValid())
+		{
+			RegisterMaskAndWeight(MaskData->BlendMask.Get(), Weight);
+		}
+	}
+}
+
+void FMovieSceneMaskBlendProvider::RegisterMaskAndWeight(TObjectPtr<UUAFBlendMask> BlendMask, const double EasingWeight)
+{
+	if (!BlendMask)
+	{
+		return;
+	}
+
+	if (!MaskSkeleton.IsValid())
+	{
+		MaskSkeleton = BlendMask->GetSkeleton();
+	}
+
+	const TArray<float>& AssetBoneWeights = BlendMask->GetSkeletonBoneWeights();
+	const TArray<FBlendProfileStandaloneCachedData::FMaskedAttributeWeight> AssetAttributeWeights = BlendMask->GetSkeletonAttributeWeights();
+
+	if (AssetBoneWeights.Num() > AccumulationBoneWeights.Num())
+	{
+		AccumulationBoneWeights.AddDefaulted(AssetBoneWeights.Num() - AccumulationBoneWeights.Num());
+	}
+
+	for (FBlendProfileStandaloneCachedData::FMaskedAttributeWeight AttributeWeight : AssetAttributeWeights)
+	{
+		if (float* CurrentWeight = AttributeWeightsMap.Find(AttributeWeight.Attribute))
+		{
+			const float DampedValue = 1.0f - AttributeWeight.Weight;
+			const float EasingAdjustedWeight = DampedValue * EasingWeight;
+			*CurrentWeight = FMath::Clamp(*CurrentWeight + EasingAdjustedWeight, 0.0f, 1.0f);
+		}
+		else
+		{
+			const float DampedValue = 1.0f - AttributeWeight.Weight;
+			const float EasingAdjustedWeight = DampedValue * EasingWeight;
+			AttributeWeightsMap.Add(AttributeWeight.Attribute, FMath::Clamp(EasingAdjustedWeight, 0.0f, 1.0f));
+		}
+	}
+
+	// Since masking will remove the influence of a pose on bones, multiple masks cannot be blended as a series of tasks. 
+	// The blend task therefore needs to be a culmination of all the weights and tables contributing to this blend provider.
+	for (int32 BoneIndex = 0; BoneIndex < AssetBoneWeights.Num(); ++BoneIndex)
+	{
+		// First pass on calculating bone weight based on the asset data and easing results.
+		// This value is the influence a particular mask will have on a bone given the asset weight and easing weight.
+
+		// The weight in the asset is representative of what the expected value should be when easing is at 1.0 and no other masks are present
+		// The expected value can be thought of as the result of subtracting a damping value from the default weight of 1.0.
+		// i.e. Asset Weight = Default Weight - Damping Value, so the Damping Value = 1.0 - Asset Weight.
+
+		// The weight and easing results determine the degree to which the values from the mask should affect the underlying pose.
+		// These values act additively from mask to mask
+		// So the ultimate result for this bone's damping value will be the sum of (Easing Amount * Damping Value) for all masks.
+
+		const float AssetWeight = AssetBoneWeights[BoneIndex];
+		const float DampingValue = 1.0f - AssetWeight;
+		const float EasingAdjustedWeight = EasingWeight * DampingValue;
+		// Accumulate easing damping values.
+		const float CurrentBoneWeight = AccumulationBoneWeights[BoneIndex];
+		AccumulationBoneWeights[BoneIndex] = FMath::Clamp(CurrentBoneWeight + EasingAdjustedWeight, 0.0, 1.0);
+	}
+
+	auto GatherCurveValues = [this, &EasingWeight](UE::Anim::FCurveElement CurveElement)
+	{
+		if (float* CurrentValue = CurveWeightsMap.Find(CurveElement.Name))
+		{
+			const float DampedValue = 1.0 - CurveElement.Value;
+			const float EasingAdjustedValue = EasingWeight * DampedValue;
+			*CurrentValue = FMath::Clamp(*CurrentValue + EasingAdjustedValue, 0.0f, 1.0f);
+		}
+		else
+		{
+			const float DampedValue = 1.0 - CurveElement.Value;
+			const float EasingAdjustedValue = EasingWeight * DampedValue;
+			CurveWeightsMap.Add(CurveElement.Name, FMath::Clamp(EasingAdjustedValue, 0.0f, 1.0f));
+		}
+	};
+
+	BlendMask->GetSkeletonCurveWeights().ForEachElement(GatherCurveValues);
+}
+
+void FMovieSceneMaskBlendProvider::ResetAccumulationData()
+{
+	AccumulationBoneWeights.Reset();
+	AttributeWeightsMap.Reset();
+	CurveWeightsMap.Reset();
+	MaskSkeleton = nullptr;
+}
+
+void FMovieSceneMaskBlendProvider::ResetTaskData()
+{
+	TaskBoneWeights.Reset();
+	TaskAttributeWeights.Reset();
+	TaskCurveWeights.Empty();
+}
+
+
+UMovieSceneAnimationMaskDecoration::UMovieSceneAnimationMaskDecoration(const FObjectInitializer& ObjInit)
+	: Super(ObjInit)
+{
+
+}
+
+UMovieSceneSection* UMovieSceneAnimationMaskDecoration::CreateNewSection()
+{
+	return NewObject<UMovieSceneAnimMixerMaskSection>(this, UMovieSceneAnimMixerMaskSection::StaticClass(), NAME_None, RF_Transactional);
+}
+
+void UMovieSceneAnimationMaskDecoration::AddSection(UMovieSceneSection* InSection)
+{
+	UMovieSceneAnimMixerMaskSection* NewMaskSection = Cast<UMovieSceneAnimMixerMaskSection>(InSection);
+	if (!NewMaskSection)
+	{
+		return;
+	}
+
+	const TRange<FFrameNumber> NewSectionRange = NewMaskSection->GetRange();
+
+	if (NewSectionRange.GetLowerBound().IsClosed() && NewSectionRange.GetUpperBound().IsClosed())
+	{
+		for (UMovieSceneSection* Section : MaskSections)
+		{
+			const TRange<FFrameNumber> CurrentSectionRange = Section->GetRange();
+			if (CurrentSectionRange.GetUpperBound().IsOpen() || CurrentSectionRange.GetLowerBound().IsOpen())
+			{
+				continue;
+			}
+			if (CurrentSectionRange.Overlaps(NewSectionRange))
+			{
+				if (Section->GetInclusiveStartFrame() < NewMaskSection->GetInclusiveStartFrame())
+				{
+					Section->SetEndFrame(NewMaskSection->GetInclusiveStartFrame());
+				}
+				else
+				{
+					NewMaskSection->SetEndFrame(Section->GetInclusiveStartFrame());
+				}
+			}
+		}
+	}
+
+	MaskSections.Add(NewMaskSection);
+
+	MarkAsChanged();
+#if WITH_EDITOR
+	MarkStructureChanged();
+#endif
+
+	NewMaskSection->SetAnimMixerLayer(GetTypedOuter<UMovieSceneAnimationMixerLayer>());
+}
+
+void UMovieSceneAnimationMaskDecoration::RemoveSection(UMovieSceneSection& SectionToRemove)
+{
+	MaskSections.RemoveAll([&SectionToRemove](const UMovieSceneSection* InSection){ return InSection == &SectionToRemove; });
+	MarkAsChanged();
+#if WITH_EDITOR
+	MarkStructureChanged();
+#endif
+}
diff --git a/Engine/Plugins/Experimental/MovieSceneAnimMixer/Source/MovieSceneAnimMixer/Private/MovieSceneAnimationMixerTrack.cpp b/Engine/Plugins/Experimental/MovieSceneAnimMixer/Source/MovieSceneAnimMixer/Private/MovieSceneAnimationMixerTrack.cpp
--- a/Engine/Plugins/Experimental/MovieSceneAnimMixer/Source/MovieSceneAnimMixer/Private/MovieSceneAnimationMixerTrack.cpp
+++ b/Engine/Plugins/Experimental/MovieSceneAnimMixer/Source/MovieSceneAnimMixer/Private/MovieSceneAnimationMixerTrack.cpp
@@ -15,10 +15,13 @@
 #include "MovieSceneAnimationTrackDecoration.h"
 #include "MovieSceneAnimationMixerLayer.h"
 #include "MovieScene.h"
+#include "MovieSceneAnimMixerMaskSection.h"
+#include "Decorations/IMovieSceneSectionProviderDecoration.h"
 #include "MovieSceneSequence.h"
 #include "MovieSceneTrack.h"
 #include "AnimMixerBakeEvaluation.h"
 #include "MovieSceneRootMotionSection.h"
+#include "Compilation/MovieSceneEvaluationTreePopulationRules.h"
 
 #define LOCTEXT_NAMESPACE "MovieSceneAnimationMixerTrack"
 
@@ -330,6 +333,23 @@ bool UMovieSceneAnimationMixerTrack::SupportsType(TSubclassOf<UMovieSceneSection
 		|| SectionClass.Get()->ImplementsInterface(UMovieSceneAnimationMixerItemInterface::StaticClass());
 }
 
+bool UMovieSceneAnimationMixerTrack::PopulateEvaluationTree(TMovieSceneEvaluationTree<FMovieSceneTrackEvaluationData>& OutData) const
+{
+	UE::MovieScene::FEvaluationTreePopulationRules::HighPassPerRow(AnimationSections, OutData);
+
+	for (const UMovieSceneAnimationMixerLayer* Layer : GetLayers())
+	{
+		for (UObject* Decoration : Layer->GetDecorations())
+		{
+			if (IMovieSceneSectionProviderDecoration* SectionProvider = Cast<IMovieSceneSectionProviderDecoration>(Decoration))
+			{
+				UE::MovieScene::FEvaluationTreePopulationRules::Blended(SectionProvider->GetSections(), OutData);
+			}
+		}
+	}
+	return true;
+}
+
 bool UMovieSceneAnimationMixerTrack::FixRowIndices()
 {
 	// Do nothing with row indices here- our rows are layers and we can have empty ones.
@@ -947,12 +967,146 @@ EMovieSceneSectionMovedResult UMovieSceneAnimationMixerTrack::OnSectionMoved(UMo
 
 EMovieSceneTrackEasingSupportFlags UMovieSceneAnimationMixerTrack::SupportsEasing(FMovieSceneSupportsEasingParams& Params) const
 {
+	if (Params.ForSection && Params.ForSection->IsA(UMovieSceneAnimMixerMaskSection::StaticClass()))
+	{
+		return EMovieSceneTrackEasingSupportFlags::All;
+	}
 	// Disable manual easing since the mixer track uses transition sections for blending
 	// instead of per-section easing. This prevents users from accidentally setting up
 	// easing that would conflict with transition blends.
 	return EMovieSceneTrackEasingSupportFlags::None;
 }
 
+void UMovieSceneAnimationMixerTrack::UpdateEasing()
+{
+	// Update easing on a per-section basis. Skeletal animations are handled by transitions within a layer, easing is only applied to sections within
+	// section provider decorations.
+
+	for (TObjectPtr<UMovieSceneAnimationMixerLayer> Layer : Layers)
+	{
+		if (!Layer)
+		{
+			continue;
+		}
+		for (TObjectPtr<UObject> Decoration : Layer.Get()->GetDecorations())
+		{
+			if (IMovieSceneSectionProviderDecoration* SectionProvider = Cast<IMovieSceneSectionProviderDecoration>(Decoration))
+			{
+				TArrayView<TObjectPtr<UMovieSceneSection>> Sections = SectionProvider->GetSections();
+				int32 MaxRowIndex = 0;
+				TArray<UMovieSceneSection*> RowSections;
+				for (TObjectPtr<UMovieSceneSection> Section : Sections)
+				{
+					MaxRowIndex = FMath::Max(MaxRowIndex, Section->GetRowIndex());
+				}
+				for (int32 RowIndex = 0; RowIndex <= MaxRowIndex; ++RowIndex)
+				{
+					RowSections.Reset();
+					for (TObjectPtr<UMovieSceneSection> Section : Sections)
+					{
+						if (Section && Section->GetRowIndex() == RowIndex)
+						{
+							RowSections.Add(Section);
+						}
+					}
+					for (int32 SectionIndex = 0; SectionIndex < RowSections.Num(); ++SectionIndex)
+					{
+						UMovieSceneSection* CurrentSection = RowSections[SectionIndex];
+
+						FMovieSceneSupportsEasingParams SupportsEasingParams(CurrentSection);
+						EMovieSceneTrackEasingSupportFlags EasingFlags = SupportsEasing(SupportsEasingParams);
+
+						// Auto-deactivate manual easing if we lost the ability to use it.
+						if (!EnumHasAllFlags(EasingFlags, EMovieSceneTrackEasingSupportFlags::ManualEaseIn))
+						{
+							CurrentSection->Easing.bManualEaseIn = false;
+						}
+						if (!EnumHasAllFlags(EasingFlags, EMovieSceneTrackEasingSupportFlags::ManualEaseOut))
+						{
+							CurrentSection->Easing.bManualEaseOut = false;
+						}
+
+						if (!EnumHasAllFlags(EasingFlags, EMovieSceneTrackEasingSupportFlags::AutomaticEasing))
+						{
+							if (CurrentSection->Easing.AutoEaseInDuration != 0 || CurrentSection->Easing.AutoEaseOutDuration != 0)
+							{
+								CurrentSection->Modify();
+
+								CurrentSection->Easing.AutoEaseInDuration = 0;
+								CurrentSection->Easing.AutoEaseOutDuration = 0;
+							}
+							continue;
+						}
+
+						// Check overlaps with exclusive ranges so that sections can butt up against each other
+						UMovieSceneTrack* OuterTrack = CurrentSection->GetTypedOuter<UMovieSceneTrack>();
+						int32 MaxEaseIn = 0;
+						int32 MaxEaseOut = 0;
+						bool bIsEntirelyUnderlapped = false;
+
+						TRange<FFrameNumber> CurrentSectionRange = CurrentSection->GetRange();
+						for (int32 OtherIndex = 0; OtherIndex < RowSections.Num(); ++OtherIndex)
+						{
+							if (OtherIndex == SectionIndex)
+							{
+								continue;
+							}
+
+							UMovieSceneSection* Other = RowSections[OtherIndex];
+							TRange<FFrameNumber> OtherSectionRange = Other->GetRange();
+
+							if (!OtherSectionRange.HasLowerBound() && !OtherSectionRange.HasUpperBound())
+							{
+								// If we're testing against an infinite range we want to use the PlayRange of the sequence
+								// instead so that blends stop at the end of a clip instead of a quarter of the length.
+								UMovieScene* OuterScene = OuterTrack->GetTypedOuter<UMovieScene>();
+								OtherSectionRange = OuterScene->GetPlaybackRange();
+							}
+
+							bIsEntirelyUnderlapped = bIsEntirelyUnderlapped || OtherSectionRange.Contains(CurrentSectionRange);
+
+							// Check the lower bound of the current section against the other section's upper bound
+							const bool bSectionRangeContainsOtherUpperBound = !OtherSectionRange.GetUpperBound().IsOpen() && !CurrentSectionRange.GetLowerBound().IsOpen() && CurrentSectionRange.Contains(OtherSectionRange.GetUpperBoundValue());
+							const bool bSectionRangeContainsOtherLowerBound = !OtherSectionRange.GetLowerBound().IsOpen() && !CurrentSectionRange.GetUpperBound().IsOpen() && CurrentSectionRange.Contains(OtherSectionRange.GetLowerBoundValue());
+							if (bSectionRangeContainsOtherUpperBound && !bSectionRangeContainsOtherLowerBound)
+							{
+								const int32 Difference = UE::MovieScene::DiscreteSize(TRange<FFrameNumber>(CurrentSectionRange.GetLowerBound(), OtherSectionRange.GetUpperBound()));
+								MaxEaseIn = FMath::Max(MaxEaseIn, Difference);
+							}
+
+							if (bSectionRangeContainsOtherLowerBound &&!bSectionRangeContainsOtherUpperBound)
+							{
+								const int32 Difference = UE::MovieScene::DiscreteSize(TRange<FFrameNumber>(OtherSectionRange.GetLowerBound(), CurrentSectionRange.GetUpperBound()));
+								MaxEaseOut = FMath::Max(MaxEaseOut, Difference);
+							}
+						}
+
+						const bool  bIsFinite = CurrentSectionRange.HasLowerBound() && CurrentSectionRange.HasUpperBound();
+						const int32 MaxSize   = bIsFinite ? UE::MovieScene::DiscreteSize(CurrentSectionRange) : TNumericLimits<int32>::Max();
+
+						if (MaxEaseOut == 0 && MaxEaseIn == 0 && bIsEntirelyUnderlapped)
+						{
+							MaxEaseOut = MaxEaseIn = MaxSize / 4;
+						}
+
+						// Only modify the section if the ease in or out times have actually changed
+						MaxEaseIn  = FMath::Clamp(MaxEaseIn, 0, MaxSize);
+						MaxEaseOut = FMath::Clamp(MaxEaseOut, 0, MaxSize);
+
+						if (CurrentSection->Easing.AutoEaseInDuration != MaxEaseIn || CurrentSection->Easing.AutoEaseOutDuration != MaxEaseOut)
+						{
+							CurrentSection->Modify();
+
+							CurrentSection->Easing.AutoEaseInDuration  = MaxEaseIn;
+							CurrentSection->Easing.AutoEaseOutDuration = MaxEaseOut;
+						}
+					}
+				}
+			}
+		}
+	}
+}
+
 #if WITH_EDITORONLY_DATA
 
 FText UMovieSceneAnimationMixerTrack::GetTrackRowDisplayName(int32 RowIndex) const
diff --git a/Engine/Plugins/Experimental/MovieSceneAnimMixer/Source/MovieSceneAnimMixer/Private/Systems/MovieSceneAnimMixerSystem.cpp b/Engine/Plugins/Experimental/MovieSceneAnimMixer/Source/MovieSceneAnimMixer/Private/Systems/MovieSceneAnimMixerSystem.cpp
--- a/Engine/Plugins/Experimental/MovieSceneAnimMixer/Source/MovieSceneAnimMixer/Private/Systems/MovieSceneAnimMixerSystem.cpp
+++ b/Engine/Plugins/Experimental/MovieSceneAnimMixer/Source/MovieSceneAnimMixer/Private/Systems/MovieSceneAnimMixerSystem.cpp
@@ -52,6 +52,8 @@
 #include "MovieSceneAnimTransitionSectionBase.h"
 #include "EvaluationVM/Tasks/TransitionEvaluationTask.h"
 #include "EntitySystem/MovieSceneInstanceRegistry.h"
+#include "MovieSceneAnimationMaskDecoration.h"
+#include "MovieSceneAnimationMixerLayer.h"
 #include UE_INLINE_GENERATED_CPP_BY_NAME(MovieSceneAnimMixerSystem)
 
 namespace UE::MovieScene
@@ -581,7 +583,15 @@ static TSharedPtr<UE::UAF::FEvaluationProgram> BuildProgramFromEntries(
 				}
 				else
 				{
-					Program->AppendTask(FAnimNextBlendTwoKeyframesPreserveRootMotionTask::Make(MixerEntry->PoseWeight));
+					if (MixerEntry->BlendTask.IsValid())
+					{
+						Program->AppendTaskPtr(MixerEntry->BlendTask.Pin());
+						MixerEntry->BlendTask = nullptr;
+					}
+					else
+					{
+						Program->AppendTask(FAnimNextBlendTwoKeyframesPreserveRootMotionTask::Make(MixerEntry->PoseWeight));
+					}
 				}
 			}
 		}
@@ -635,6 +645,64 @@ namespace UE::MovieScene
 		}
 	};
 
+	struct FGatherBlendTasks
+	{
+		TMap<TWeakObjectPtr<UMovieSceneAnimationMixerLayer>, TWeakPtr<FMovieSceneAnimMixerEntry>>* EntriesByLayer;
+
+		FGatherBlendTasks(TMap<TWeakObjectPtr<UMovieSceneAnimationMixerLayer>, TWeakPtr<FMovieSceneAnimMixerEntry>>* InEntriesByLayer)
+			: EntriesByLayer(InEntriesByLayer)
+		{
+		}
+		~FGatherBlendTasks()
+		{
+		}
+		void ForEachEntity(TObjectPtr<UMovieSceneAnimationMixerLayer> MixerLayer, const TInstancedStruct<FMovieSceneAnimMixerBlendProviderData>& BlendData, const double* WeightAndEasing) const
+		{
+			const double Weight = WeightAndEasing ? *WeightAndEasing : 1.0f;
+			for (TObjectPtr<UObject> DecorationPtr : MixerLayer.Get()->GetDecorations())
+			{
+				if (IAnimMixerBlendProviderDecoration* BlendProviderDecoration = Cast<IAnimMixerBlendProviderDecoration>(DecorationPtr))
+				{
+					if (FMovieSceneAnimMixerBlendProvider* BlendProvider = BlendProviderDecoration->GetBlendProvider())
+					{
+						BlendProvider->RegisterTaskData(BlendData, Weight);
+					}
+				}
+			}
+		}
+		void PostTask()
+		{
+			TArray<TWeakObjectPtr<UMovieSceneAnimationMixerLayer>> MixerLayers;
+			EntriesByLayer->GetKeys(MixerLayers);
+			for (TWeakObjectPtr<UMovieSceneAnimationMixerLayer> MixerLayerKey : MixerLayers)
+			{
+				const TWeakPtr<FMovieSceneAnimMixerEntry>* Entry = EntriesByLayer->Find(MixerLayerKey);
+
+				const UMovieSceneAnimationMixerLayer* MixerLayer = MixerLayerKey.Get();
+				if (!MixerLayer)
+				{
+					continue;
+				}
+				for (TObjectPtr<UObject> DecorationPtr : MixerLayer->GetDecorations())
+				{
+					if (IAnimMixerBlendProviderDecoration* BlendProviderDecoration = Cast<IAnimMixerBlendProviderDecoration>(DecorationPtr))
+					{
+						if (FMovieSceneAnimMixerBlendProvider* BlendProvider = BlendProviderDecoration->GetBlendProvider())
+						{
+							if (Entry && Entry->IsValid() && BlendProvider)
+							{
+								Entry->Pin()->BlendTask = BlendProvider->GetBlendTask();
+								// Only one blend task per provider. The first blend provider is the one to be used.
+								break;
+							}
+						}
+					}
+				}
+			}
+			MixerLayers.Reset();
+		}
+	};
+
 	struct FEvaluateAnimMixers
 	{
 		TMap<FMovieSceneAnimMixerKey, TSharedPtr<FMovieSceneAnimMixer>>* Mixers;
@@ -1204,6 +1272,7 @@ void UMovieSceneAnimMixerSystem::OnRun(FSystemTaskPrerequisites& InPrerequisites
 		.ReadOptional(AnimMixerComponents->RootMotionSettings)
 		.ReadOptional(AnimMixerComponents->MirrorTable)
 		.ReadOptional(AnimMixerComponents->EntityOwner)
+		.ReadOptional(AnimMixerComponents->MixerLayer)
 		.FilterAll({ BuiltInComponents->Tags.NeedsLink})
 		.Iterate_PerEntity( &Linker->EntityManager,
 		[this, &NewMixers, BuiltInComponents, AnimMixerComponents]
@@ -1218,7 +1287,8 @@ void UMovieSceneAnimMixerSystem::OnRun(FSystemTaskPrerequisites& InPrerequisites
 			const FMovieSceneSkeletalAnimationComponentData* SkelAnimData,
 			const FMovieSceneRootMotionSettings* RootMotionSettings,
 			const TObjectPtr<UMirrorDataTable>* OptMirrorDataTable,
-			const FObjectKey* EntityOwner)
+			const FObjectKey* EntityOwner,
+			const TObjectPtr<UMovieSceneAnimationMixerLayer>* MixerLayer)
 			{
 				// For new entities, we 'resolve' the animation target so if 'Automatic' is picked we choose the right one automatically.
 				Target = ResolveAnimationTarget(BoundObjectKey, Target, SkelAnimData);
@@ -1243,6 +1313,11 @@ void UMovieSceneAnimMixerSystem::OnRun(FSystemTaskPrerequisites& InPrerequisites
 
 				FMovieSceneAnimMixerEntry* Entry = InOutMixerEntry.Get();
 
+				if (MixerLayer && MixerLayer->Get())
+				{
+					EntriesByLayer.Add(*MixerLayer, InOutMixerEntry);
+				}
+
 				Entry->InstanceHandle = InstanceHandle;
 				Entry->EntityID = EntityID;
 				Entry->EvalTask = Task;
@@ -1482,6 +1557,12 @@ void UMovieSceneAnimMixerSystem::OnSchedulePersistentTasks(UE::MovieScene::IEnti
 		.FilterNone({BuiltInComponents->Tags.NeedsUnlink, BuiltInComponents->Tags.Ignored})
 		.Schedule_PerEntity<FUpdateTaskEntities>(&Linker->EntityManager, TaskScheduler);
 
+	FTaskID BlendFinderTask = FEntityTaskBuilder()
+		.Read(AnimMixerComponents->MixerLayer)
+		.Read(AnimMixerComponents->BlendProviderData)
+		.ReadOptional(BuiltInComponents->WeightAndEasingResult)
+		.FilterNone({BuiltInComponents->Tags.NeedsLink, BuiltInComponents->Tags.Ignored})
+		.Schedule_PerEntity<FGatherBlendTasks>(&Linker->EntityManager, TaskScheduler, &EntriesByLayer);
 
 	// For each mixer, build the evaluation program task list.
 	FTaskID MixTask = FEntityTaskBuilder()
@@ -1492,6 +1573,7 @@ void UMovieSceneAnimMixerSystem::OnSchedulePersistentTasks(UE::MovieScene::IEnti
 		.Schedule_PerEntity<FEvaluateAnimMixers>(&Linker->EntityManager, TaskScheduler, &Mixers, &RootMotionByObject, Linker, this);
 
 	TaskScheduler->AddPrerequisite(UpdateTask, MixTask);
+	TaskScheduler->AddPrerequisite(BlendFinderTask, MixTask);
 }
 
 void UMovieSceneAnimMixerSystem::OnCleanTaggedGarbage()
diff --git a/Engine/Plugins/Experimental/MovieSceneAnimMixer/Source/MovieSceneAnimMixer/Public/AnimMixerComponentTypes.h b/Engine/Plugins/Experimental/MovieSceneAnimMixer/Source/MovieSceneAnimMixer/Public/AnimMixerComponentTypes.h
--- a/Engine/Plugins/Experimental/MovieSceneAnimMixer/Source/MovieSceneAnimMixer/Public/AnimMixerComponentTypes.h
+++ b/Engine/Plugins/Experimental/MovieSceneAnimMixer/Source/MovieSceneAnimMixer/Public/AnimMixerComponentTypes.h
@@ -11,6 +11,7 @@
 #include "StructUtils/InstancedStruct.h"
 #include "EvaluationVM/EvaluationTask.h"
 #include "AnimSequencerInstanceProxy.h"
+#include "MovieSceneAnimMixerBlendProvider.h"
 #include "Systems/MovieSceneRootMotionSystem.h"
 #include "Channels/MovieSceneInterpolation.h"
 
@@ -38,6 +39,7 @@ namespace UE::MovieScene
 		TComponentTypeID<TInstancedStruct<FMovieSceneMixedAnimationTarget>> Target;
 		TComponentTypeID<TSharedPtr<FAnimNextEvaluationTask>> Task;
 		TComponentTypeID<TSharedPtr<FAnimNextEvaluationTask>> MixerTask;
+		TComponentTypeID<TInstancedStruct<FMovieSceneAnimMixerBlendProviderData>> BlendProviderData;
 		TComponentTypeID<TSharedPtr<FMovieSceneAnimMixerEntry>> MixerEntry;
 		TComponentTypeID<FMovieSceneRootMotionSettings> RootMotionSettings;
 		TComponentTypeID<EMovieSceneRootMotionDestination> RootDestination;
diff --git a/Engine/Plugins/Experimental/MovieSceneAnimMixer/Source/MovieSceneAnimMixer/Public/MovieSceneAnimMixerBlendProvider.h b/Engine/Plugins/Experimental/MovieSceneAnimMixer/Source/MovieSceneAnimMixer/Public/MovieSceneAnimMixerBlendProvider.h
new file mode 100644
--- /dev/null
+++ b/Engine/Plugins/Experimental/MovieSceneAnimMixer/Source/MovieSceneAnimMixer/Public/MovieSceneAnimMixerBlendProvider.h
@@ -0,0 +1,52 @@
+﻿// Copyright Epic Games, Inc. All Rights Reserved.
+
+#pragma once
+
+#include "Templates/SharedPointer.h"
+#include "StructUtils/InstancedStruct.h"
+#include "UObject/Interface.h"
+
+#include "MovieSceneAnimMixerBlendProvider.generated.h"
+
+struct FAnimNextEvaluationTask;
+
+USTRUCT()
+struct FMovieSceneAnimMixerBlendProviderData
+{
+	GENERATED_BODY()
+};
+
+/**
+ * Defines a Mixer animation blend.
+ */
+USTRUCT()
+struct FMovieSceneAnimMixerBlendProvider
+{
+	GENERATED_BODY()
+	
+	virtual ~FMovieSceneAnimMixerBlendProvider() = default;
+	
+	TSharedPtr<FAnimNextEvaluationTask> BlendTask = nullptr;
+
+	virtual void RegisterTaskData(const TInstancedStruct<FMovieSceneAnimMixerBlendProviderData>& InData, const double Weight) {}
+
+	virtual TSharedPtr<FAnimNextEvaluationTask> GetBlendTask() { return BlendTask; }
+};
+
+UINTERFACE(MinimalAPI)
+class UAnimMixerBlendProviderDecoration : public UInterface
+{
+public:
+	GENERATED_BODY()
+};
+
+/**
+ * Optional interface for movie scene decorations that provide a blend for animation mixer layers.
+ */
+class IAnimMixerBlendProviderDecoration
+{
+public:
+	GENERATED_BODY()
+
+	virtual FMovieSceneAnimMixerBlendProvider* GetBlendProvider() { return nullptr; }
+};
\ No newline at end of file
diff --git a/Engine/Plugins/Experimental/MovieSceneAnimMixer/Source/MovieSceneAnimMixer/Public/MovieSceneAnimMixerMaskSection.h b/Engine/Plugins/Experimental/MovieSceneAnimMixer/Source/MovieSceneAnimMixer/Public/MovieSceneAnimMixerMaskSection.h
new file mode 100644
--- /dev/null
+++ b/Engine/Plugins/Experimental/MovieSceneAnimMixer/Source/MovieSceneAnimMixer/Public/MovieSceneAnimMixerMaskSection.h
@@ -0,0 +1,52 @@
+﻿// Copyright Epic Games, Inc. All Rights Reserved.
+
+#pragma once
+#include <UAF/BlendMask/UAFBlendMask.h>
+
+#include "HierarchyTable.h"
+#include "MovieSceneAnimationMixerItemInterface.h"
+#include "MovieSceneAnimationMixerLayer.h"
+#include "MovieSceneSection.h"
+#include "Channels/MovieSceneObjectPathChannel.h"
+#include "EntitySystem/IMovieSceneEntityProvider.h"
+
+#include "MovieSceneAnimMixerMaskSection.generated.h"
+
+struct FMovieSceneMaskBlendProvider;
+
+UCLASS(MinimalAPI, DisplayName="Mask")
+class UMovieSceneAnimMixerMaskSection
+	: public UMovieSceneSection
+	, public IMovieSceneEntityProvider
+	, public IMovieSceneAnimationMixerItemInterface
+{
+public:
+
+	UMovieSceneAnimMixerMaskSection();
+
+	GENERATED_BODY()
+
+	void SetAnimMixerLayer(UMovieSceneAnimationMixerLayer* InLayer) { OwningLayer = InLayer; }
+
+	TObjectKey<UUAFBlendMask> GetBlendMask() { return BlendMask; }
+
+#if WITH_EDITOR
+	virtual bool SupportsVerticalDragging() const override { return false; }
+
+	virtual bool IsVisibleInAddSectionMenu() const override {return false; }
+#endif
+
+private:
+
+	UPROPERTY(EditAnywhere, Category="Default")
+	TObjectPtr<UUAFBlendMask> BlendMask = nullptr;
+
+	TSharedPtr<FMovieSceneMaskBlendProvider> BlendProvider;
+
+	UPROPERTY()
+	TObjectPtr<UMovieSceneAnimationMixerLayer> OwningLayer;
+
+	virtual void ImportEntityImpl(UMovieSceneEntitySystemLinker* EntityLinker, const FEntityImportParams& Params, FImportedEntity* OutImportedEntity) override;
+
+	FColor MaskTintColor = FColor(110, 25, 50);
+};
diff --git a/Engine/Plugins/Experimental/MovieSceneAnimMixer/Source/MovieSceneAnimMixer/Public/MovieSceneAnimationMaskDecoration.h b/Engine/Plugins/Experimental/MovieSceneAnimMixer/Source/MovieSceneAnimMixer/Public/MovieSceneAnimationMaskDecoration.h
new file mode 100644
--- /dev/null
+++ b/Engine/Plugins/Experimental/MovieSceneAnimMixer/Source/MovieSceneAnimMixer/Public/MovieSceneAnimationMaskDecoration.h
@@ -0,0 +1,98 @@
+﻿// Copyright Epic Games, Inc. All Rights Reserved.
+
+#pragma once
+
+#include "HierarchyTableBlendProfile.h"
+#include "MovieSceneAnimMixerBlendProvider.h"
+#include "MovieSceneSignedObject.h"
+#include "Decorations/IMovieSceneSectionProviderDecoration.h"
+#include "Systems/MovieSceneAnimMixerSystem.h"
+
+#include "MovieSceneAnimationMaskDecoration.generated.h"
+
+class UMovieSceneAnimationMixerLayer;
+class UUAFBlendMask;
+class UMovieSceneTrack;
+
+USTRUCT()
+struct FMovieSceneMaskBlendProviderData : public FMovieSceneAnimMixerBlendProviderData
+{
+	GENERATED_BODY()
+	TWeakObjectPtr<UUAFBlendMask> BlendMask;
+};
+
+USTRUCT()
+struct FMovieSceneMaskBlendProvider : public FMovieSceneAnimMixerBlendProvider
+{
+	GENERATED_BODY()
+
+	FMovieSceneMaskBlendProvider() {}
+
+	virtual ~FMovieSceneMaskBlendProvider() = default;
+
+	virtual TSharedPtr<FAnimNextEvaluationTask> GetBlendTask() override;
+
+	virtual void RegisterTaskData(const TInstancedStruct<FMovieSceneAnimMixerBlendProviderData>& InData, const double Weight) override;
+
+	void RegisterMaskAndWeight(TObjectPtr<UUAFBlendMask> BlendMask, const double EasingWeight);
+
+
+private:
+
+	void ResetAccumulationData();
+	void ResetTaskData();
+
+	// Skeleton Associated with the blend task.
+	TWeakObjectPtr<USkeleton> MaskSkeleton;
+
+	// Per bone weights based on values from blend mask, and the section's weight and easing, additively accumulated from all tables.
+	TArray<float> AccumulationBoneWeights;
+	// Bone weights to be passed to the UAF task.
+	TArray<float> TaskBoneWeights;
+
+	// Curve weights, similar to AccumulationBoneWeights, will be additively accumulated from sections based on their weight and easing.
+	TMap<FName, float> CurveWeightsMap;
+	// Named Value Array passed to the UAF Task
+	UE::Anim::TNamedValueArray<FDefaultAllocator, UE::Anim::FCurveElement> TaskCurveWeights;
+
+	// Attribute weights, also additively accumulated based on section weight and easing.
+	TMap<UE::Anim::FAttributeId, float> AttributeWeightsMap;
+	// Array passed to the UAF Task
+	TArray<FBlendProfileStandaloneCachedData::FMaskedAttributeWeight> TaskAttributeWeights;
+};
+
+/**
+ * Decorator for associating a mixer target to a per-bone blend mask (UAF Blend Mask). The anim mixer will then apply the animation only to the bones specified.
+ * Bone, attribute and curve wieghts can be blended using sequencer easing.
+ */
+UCLASS(MinimalAPI, DisplayName="Mask")
+class UMovieSceneAnimationMaskDecoration
+	: public UMovieSceneSignedObject
+	, public IMovieSceneSectionProviderDecoration
+	, public IAnimMixerBlendProviderDecoration
+{
+public:
+	GENERATED_BODY()
+
+	UMovieSceneAnimationMaskDecoration(const FObjectInitializer& ObjInit);
+
+	virtual TArrayView<TObjectPtr<UMovieSceneSection>> GetSections() override { return MaskSections;}
+
+	virtual UMovieSceneSection* CreateNewSection() override;
+
+	MOVIESCENEANIMMIXER_API virtual void AddSection(UMovieSceneSection* InSection) override;
+
+	MOVIESCENEANIMMIXER_API virtual void RemoveSection(UMovieSceneSection& SectionToRemove) override;
+
+	MOVIESCENEANIMMIXER_API virtual bool SupportsMultipleSections() override { return true; }
+
+	MOVIESCENEANIMMIXER_API virtual FMovieSceneAnimMixerBlendProvider* GetBlendProvider() override { return &BlendProvider; }
+
+	UPROPERTY()
+	TArray<TObjectPtr<UMovieSceneSection>> MaskSections;
+
+private:
+	UPROPERTY()
+	FMovieSceneMaskBlendProvider BlendProvider;
+};
+
diff --git a/Engine/Plugins/Experimental/MovieSceneAnimMixer/Source/MovieSceneAnimMixer/Public/MovieSceneAnimationMixerTrack.h b/Engine/Plugins/Experimental/MovieSceneAnimMixer/Source/MovieSceneAnimMixer/Public/MovieSceneAnimationMixerTrack.h
--- a/Engine/Plugins/Experimental/MovieSceneAnimMixer/Source/MovieSceneAnimMixer/Public/MovieSceneAnimationMixerTrack.h
+++ b/Engine/Plugins/Experimental/MovieSceneAnimMixer/Source/MovieSceneAnimMixer/Public/MovieSceneAnimationMixerTrack.h
@@ -12,6 +12,7 @@
 #include "MovieSceneMixedAnimationTarget.h"
 #include "MovieSceneNameableTrack.h"
 #include "MovieSceneAnimationMixerItemInterface.h"
+#include "Channels/MovieSceneObjectPathChannel.h"
 #include "StructUtils/InstancedStruct.h"
 #include "Tracks/MovieSceneCommonAnimationTrack.h"
 #include "Systems/MovieSceneRootMotionSystem.h"
@@ -182,13 +183,14 @@ class UMovieSceneAnimationMixerTrack
 
 public:
 
-	UMovieSceneAnimationMixerTrack(const FObjectInitializer& Init);
+	MOVIESCENEANIMMIXER_API UMovieSceneAnimationMixerTrack(const FObjectInitializer& Init);
 
 	virtual bool FixRowIndices() override;
 	virtual bool SupportsType(TSubclassOf<UMovieSceneSection> SectionClass) const override;
 
 	//~ IMovieSceneTrackVirtualAPI interface
 	virtual void RemoveSection(UMovieSceneSection& Section) override;
+	virtual bool PopulateEvaluationTree(TMovieSceneEvaluationTree<FMovieSceneTrackEvaluationData>& OutData) const override;
 
 	virtual void OnSectionAddedImpl(UMovieSceneSection* Secton) override;
 	virtual void OnSectionRemovedImpl(UMovieSceneSection* Section) override;
@@ -197,6 +199,8 @@ public:
 #endif
 	virtual EMovieSceneTrackEasingSupportFlags SupportsEasing(FMovieSceneSupportsEasingParams& Params) const override;
 
+	virtual void UpdateEasing() override;
+
 	MOVIESCENEANIMMIXER_API virtual UMovieSceneSection* AddNewAnimationOnRow(FFrameNumber KeyTime, UAnimSequenceBase* AnimSequence, int32 RowIndex) override;
 	MOVIESCENEANIMMIXER_API virtual void AddSectionWithAutoRow(UMovieSceneSection* Section) override;
 	MOVIESCENEANIMMIXER_API virtual int32 ComputeNextAvailableRowIndex() const override;
diff --git a/Engine/Plugins/Experimental/MovieSceneAnimMixer/Source/MovieSceneAnimMixer/Public/Systems/MovieSceneAnimMixerSystem.h b/Engine/Plugins/Experimental/MovieSceneAnimMixer/Source/MovieSceneAnimMixer/Public/Systems/MovieSceneAnimMixerSystem.h
--- a/Engine/Plugins/Experimental/MovieSceneAnimMixer/Source/MovieSceneAnimMixer/Public/Systems/MovieSceneAnimMixerSystem.h
+++ b/Engine/Plugins/Experimental/MovieSceneAnimMixer/Source/MovieSceneAnimMixer/Public/Systems/MovieSceneAnimMixerSystem.h
@@ -14,6 +14,7 @@
 #include "MovieSceneAnimMixerSystem.generated.h"
 
 class UMovieSceneAnimationMixerTrack;
+class UMovieSceneAnimationMixerLayer;
 struct FMovieSceneSkeletalAnimationComponentData;
 struct FAnimNextEvaluationTask;
 
@@ -29,6 +30,7 @@ struct FMovieSceneAnimMixerEntry
 
 	TWeakPtr<FMovieSceneAnimMixer> WeakMixer;
 	TSharedPtr<FAnimNextEvaluationTask> EvalTask;
+	TWeakPtr<FAnimNextEvaluationTask> BlendTask;
 	/** Shared pointer to the root motion for this entry if it came from a FMovieSceneRootMotionSettings on a mixer
 	 * @note: Only to be used for lifetime management to keep FMovieSceneAnimMixer::WeakRootMotion alive!
 	 *        This may be null even if the result of the mix still has root motion */
@@ -188,6 +190,9 @@ private:
 	// Map of animation mixers
 	TMap<FMovieSceneAnimMixerKey, TSharedPtr<FMovieSceneAnimMixer>> Mixers;
 
+	// Map of animation mixers by layer.
+	TMap<TWeakObjectPtr<UMovieSceneAnimationMixerLayer>, TWeakPtr<FMovieSceneAnimMixerEntry>> EntriesByLayer;
+
 	// Map of Root Motion data for each object
 	TMap<FObjectKey, TWeakPtr<FMovieSceneMixerRootMotionComponentData>> RootMotionByObject;
 
diff --git a/Engine/Source/Editor/Sequencer/Private/Sequencer.cpp b/Engine/Source/Editor/Sequencer/Private/Sequencer.cpp
--- a/Engine/Source/Editor/Sequencer/Private/Sequencer.cpp
+++ b/Engine/Source/Editor/Sequencer/Private/Sequencer.cpp
@@ -1968,13 +1968,27 @@ void FSequencer::DeleteSections(const TSet<UMovieSceneSection*>& Sections)
 		// Check if section comes from a decoration that provides sections
 		if (IMovieSceneSectionProviderDecoration* SectionProvider = Cast<IMovieSceneSectionProviderDecoration>(SectionOuter))
 		{
-			// Section is provided by a decoration - remove the decoration from its container
+			// Section is provided by a decoration - remove section if there are more than one, or remove the decoration from its container if it's the last section
 			UMovieSceneDecorationContainerObject* Container = Cast<UMovieSceneDecorationContainerObject>(SectionOuter->GetOuter());
 			if (Container)
 			{
 				Container->SetFlags(RF_Transactional);
 				Container->Modify();
-				Container->RemoveDecoration(SectionOuter->GetClass());
+				if (SectionProvider->SupportsMultipleSections() && SectionProvider->GetSections().Num() > 1)
+				{
+					SectionOuter->SetFlags(RF_Transactional);
+					SectionOuter->Modify();
+					SectionProvider->RemoveSection(*Section);
+					if (UMovieSceneTrack* OuterTrack = Section->GetTypedOuter<UMovieSceneTrack>())
+					{
+						OuterTrack->Modify();
+						OuterTrack->UpdateEasing();
+					}
+				}
+				else
+				{
+					Container->RemoveDecoration(SectionOuter->GetClass());
+				}
 				bAnythingRemoved = true;
 			}
 		}
diff --git a/Engine/Source/Runtime/MovieScene/Public/Decorations/IMovieSceneSectionProviderDecoration.h b/Engine/Source/Runtime/MovieScene/Public/Decorations/IMovieSceneSectionProviderDecoration.h
--- a/Engine/Source/Runtime/MovieScene/Public/Decorations/IMovieSceneSectionProviderDecoration.h
+++ b/Engine/Source/Runtime/MovieScene/Public/Decorations/IMovieSceneSectionProviderDecoration.h
@@ -2,7 +2,6 @@
 
 #pragma once
 
-#include "CoreTypes.h"
 #include "UObject/Interface.h"
 #include "IMovieSceneSectionProviderDecoration.generated.h"
 
@@ -27,4 +26,27 @@ public:
 
 	/** Returns the sections provided by this decoration */
 	virtual TArrayView<TObjectPtr<UMovieSceneSection>> GetSections() = 0;
+
+	virtual bool SupportsMultipleSections() { return false; }
+
+	virtual UMovieSceneSection* CreateNewSection() { return nullptr; }
+
+	virtual void AddSection(UMovieSceneSection* InSection) {}
+
+	virtual void RemoveSection(UMovieSceneSection& SectionToRemove) {}
+
+#if WITH_EDITOR
+	void MarkStructureChanged() { bStructureDirty = true; }
+
+	bool ConsumeStructureChanged ()
+	{
+		const bool bWasDirty = bStructureDirty;
+		bStructureDirty = false;
+		return bWasDirty;
+	}
+
+private:
+	// Have sections been added or removed, and not reflected in the UI?
+	bool bStructureDirty = false;
+#endif
 };
