Sequencer: Fix bone matching edge cases with bBlendChildOfRoot.
* Reverted an unnecessary change I had made to the bone matching logic to attempt to fix a bug that was elsewhere
* Handle the case where we have the mixer plugin enabled but using Skeletal Animation Tracks outside of the mixer with bBlendChildOfRoot selected- in this case we want to handle all root motion in the skeletal animation eval task rather than with the subsequent extract root motion task.

#Synced-CL 51609808

diff --git a/Engine/Plugins/Experimental/MovieSceneAnimMixer/Source/MovieSceneAnimMixer/Private/AnimMixerBoneMatching.cpp b/Engine/Plugins/Experimental/MovieSceneAnimMixer/Source/MovieSceneAnimMixer/Private/AnimMixerBoneMatching.cpp
--- a/Engine/Plugins/Experimental/MovieSceneAnimMixer/Source/MovieSceneAnimMixer/Private/AnimMixerBoneMatching.cpp
+++ b/Engine/Plugins/Experimental/MovieSceneAnimMixer/Source/MovieSceneAnimMixer/Private/AnimMixerBoneMatching.cpp
@@ -44,26 +44,27 @@ static FTransform ComputeMatchOffset(
 		FQuat SecBoneQuat = SectionBoneCS.GetRotation();
 		UnderlyingBoneQuat.EnforceShortestArcWith(SecBoneQuat);
 
-		// Decompose the delta between the two quaternions to Euler (the delta is small, avoiding gimbal lock),
-		// zero unmatched axes, and recompose.
-		FQuat FullDelta = SecBoneQuat.Inverse() * UnderlyingBoneQuat;
-		FullDelta.Normalize();
-		FVector DeltaEuler = EulerFromQuat(FullDelta, RotationOrder);
+		// Decompose both rotations to Euler, use SetClosestToMe to avoid
+		// winding ambiguity, then selectively copy unmatched axes.
+		FRotator UnderlyingRotation(FRotator::MakeFromEuler(EulerFromQuat(UnderlyingBoneQuat, RotationOrder)));
+		FRotator SecBoneRotation(FRotator::MakeFromEuler(EulerFromQuat(SecBoneQuat, RotationOrder)));
+		SecBoneRotation.SetClosestToMe(UnderlyingRotation);
 
 		if (!Settings.bMatchRotationZ)
 		{
-			DeltaEuler.Z = 0.0;
+			UnderlyingRotation.Yaw = SecBoneRotation.Yaw;
 		}
 		if (!Settings.bMatchRotationY)
 		{
-			DeltaEuler.Y = 0.0;
+			UnderlyingRotation.Pitch = SecBoneRotation.Pitch;
 		}
 		if (!Settings.bMatchRotationX)
 		{
-			DeltaEuler.X = 0.0;
+			UnderlyingRotation.Roll = SecBoneRotation.Roll;
 		}
 
-		DeltaRot = QuatFromEuler(DeltaEuler, RotationOrder);
+		FQuat AdjustedUnderlyingQuat = QuatFromEuler(UnderlyingRotation.Euler(), RotationOrder);
+		DeltaRot = SecBoneQuat.Inverse() * AdjustedUnderlyingQuat;
 		DeltaRot.Normalize();
 	}
 
diff --git a/Engine/Plugins/Experimental/MovieSceneAnimMixer/Source/MovieSceneAnimMixer/Private/Systems/MovieSceneAnimMixerSystem.cpp b/Engine/Plugins/Experimental/MovieSceneAnimMixer/Source/MovieSceneAnimMixer/Private/Systems/MovieSceneAnimMixerSystem.cpp
--- a/Engine/Plugins/Experimental/MovieSceneAnimMixer/Source/MovieSceneAnimMixer/Private/Systems/MovieSceneAnimMixerSystem.cpp
+++ b/Engine/Plugins/Experimental/MovieSceneAnimMixer/Source/MovieSceneAnimMixer/Private/Systems/MovieSceneAnimMixerSystem.cpp
@@ -348,7 +348,11 @@ static TSharedPtr<UE::UAF::FEvaluationProgram> BuildProgramFromEntries(
 			}
 
 			bool bNeedsMirror = (WrapEntry->MirrorTable != nullptr);
-			bool bNeedsExtraction = WrapEntry->RootMotionSettings.IsSet();
+			// Skip extraction when a legacy root motion override is active -
+			// Execute() handles the override inline and no separate extraction
+			// from the pose is needed.
+			bool bNeedsExtraction = WrapEntry->RootMotionSettings.IsSet()
+				&& !WrapEntry->RootMotionSettings->bHasRootMotionOverride;
 			TOptional<FAnimNextConvertRootMotionToWorldSpaceTask> SpaceTask;
 			if (!Filter.bSkipRootMotionConversion)
 			{
@@ -539,8 +543,10 @@ static TSharedPtr<UE::UAF::FEvaluationProgram> BuildProgramFromEntries(
 
 			// Extract root motion from the pose into a RootTransformAttribute.
 			// This runs after eval and mirror, before space conversion.
-			// Skipped for transitions (already handled in mini-programs).
-			if (MixerEntry->RootMotionSettings.IsSet() && !MixerEntry->bIsTransition)
+			// Skipped for transitions (already handled in mini-programs) and
+			// for legacy root motion overrides (handled inline by Execute()).
+			if (MixerEntry->RootMotionSettings.IsSet() && !MixerEntry->bIsTransition
+				&& !MixerEntry->RootMotionSettings->bHasRootMotionOverride)
 			{
 				Program->AppendTask(FAnimNextExtractRootMotionTask::Make(
 					MixerEntry->RootMotionSettings->RootMotionSource));
diff --git a/Engine/Plugins/Experimental/MovieSceneAnimMixer/Source/MovieSceneAnimMixer/Private/Systems/MovieSceneMixedSkeletalAnimationSystem.cpp b/Engine/Plugins/Experimental/MovieSceneAnimMixer/Source/MovieSceneAnimMixer/Private/Systems/MovieSceneMixedSkeletalAnimationSystem.cpp
--- a/Engine/Plugins/Experimental/MovieSceneAnimMixer/Source/MovieSceneAnimMixer/Private/Systems/MovieSceneMixedSkeletalAnimationSystem.cpp
+++ b/Engine/Plugins/Experimental/MovieSceneAnimMixer/Source/MovieSceneAnimMixer/Private/Systems/MovieSceneMixedSkeletalAnimationSystem.cpp
@@ -436,35 +436,56 @@ void FMovieSceneSkeletalAnimationEvaluationTask::Execute(UE::UAF::FEvaluationVM&
 			FDecompressionTools::GetAnimationPose(AnimSequencePtr, ExtractionContext, Keyframe.Pose, bUseRawData);
 
 			// Legacy root motion override: the common animation track pre-computes root motion
-			// transforms and passes them in via RootMotionOverride. Apply the override to the
-			// appropriate bone and store it as the RootTransformAttribute directly. This sets
-			// the attribute so FAnimNextExtractRootMotionTask (which runs later) will skip
-			// re-extraction.
+			// transforms and passes them in via RootMotionOverride.
 			if (AnimationData.RootMotionOverride.IsSet())
 			{
-				int32 PoseIndex = INDEX_NONE;
 				if (!AnimationData.RootMotionOverride->bBlendFirstChildOfRoot)
 				{
-					PoseIndex = Keyframe.Pose.GetRefPose().GetLODBoneIndexFromSkeletonBoneIndex(0);
+					// Standard root motion: apply the override to bone 0 and store it
+					// as the RootTransformAttribute so FAnimNextExtractRootMotionTask
+					// (which should not be created for legacy overrides) is not needed.
+					const int32 PoseIndex = Keyframe.Pose.GetRefPose().GetLODBoneIndexFromSkeletonBoneIndex(0);
+					if (PoseIndex != INDEX_NONE)
+					{
+						Keyframe.Pose.LocalTransformsView[PoseIndex] = AnimationData.RootMotionOverride->RootMotion;
+					}
+
+					FTransformAnimationAttribute* RootMotionAttribute = Keyframe.Attributes.FindOrAdd<FTransformAnimationAttribute>(AnimMixerComponents->RootTransformAttributeId);
+					FFloatAnimationAttribute* RootMotionWeight = Keyframe.Attributes.FindOrAdd<FFloatAnimationAttribute>(AnimMixerComponents->RootTransformWeightAttributeId);
+					RootMotionAttribute->Value = AnimationData.RootMotionOverride->RootMotion;
+					RootMotionWeight->Value = 1.f;
+
+					const int32 RootIndex = Keyframe.Pose.GetRefPose().GetLODBoneIndexFromSkeletonBoneIndex(0);
+					if (RootIndex != INDEX_NONE)
+					{
+						Keyframe.Pose.LocalTransformsView[RootIndex] = FTransform::Identity;
+					}
 				}
 				else if (AnimationData.RootMotionOverride->ChildBoneIndex != INDEX_NONE)
 				{
-					PoseIndex = Keyframe.Pose.GetRefPose().GetLODBoneIndexFromSkeletonBoneIndex(AnimationData.RootMotionOverride->ChildBoneIndex);
-				}
-				if (PoseIndex != INDEX_NONE)
-				{
-					Keyframe.Pose.LocalTransformsView[PoseIndex] = AnimationData.RootMotionOverride->RootMotion;
-				}
+					// Blend-first-child-of-root: the pre-computed transform targets a
+					// child bone, not the root. Apply the override to the child bone.
+					// The root bone's animated value is extracted separately as root
+					// motion so SwapRootBone still works (the old code path applied
+					// SwapRootBone independently after setting the child bone).
+					const int32 ChildPoseIndex = Keyframe.Pose.GetRefPose().GetLODBoneIndexFromSkeletonBoneIndex(AnimationData.RootMotionOverride->ChildBoneIndex);
+					if (ChildPoseIndex != INDEX_NONE)
+					{
+						Keyframe.Pose.LocalTransformsView[ChildPoseIndex] = AnimationData.RootMotionOverride->RootMotion;
+					}
 
-				FTransformAnimationAttribute* RootMotionAttribute = Keyframe.Attributes.FindOrAdd<FTransformAnimationAttribute>(AnimMixerComponents->RootTransformAttributeId);
-				FFloatAnimationAttribute* RootMotionWeight = Keyframe.Attributes.FindOrAdd<FFloatAnimationAttribute>(AnimMixerComponents->RootTransformWeightAttributeId);
-				RootMotionAttribute->Value = AnimationData.RootMotionOverride->RootMotion;
-				RootMotionWeight->Value = 1.f;
+					// Extract the root bone's animated value as root motion and zero
+					// it, mirroring what SwapRootBone did in the old code path.
+					const int32 RootIndex = Keyframe.Pose.GetRefPose().GetLODBoneIndexFromSkeletonBoneIndex(0);
+					if (RootIndex != INDEX_NONE)
+					{
+						FTransformAnimationAttribute* RootMotionAttribute = Keyframe.Attributes.FindOrAdd<FTransformAnimationAttribute>(AnimMixerComponents->RootTransformAttributeId);
+						FFloatAnimationAttribute* RootMotionWeight = Keyframe.Attributes.FindOrAdd<FFloatAnimationAttribute>(AnimMixerComponents->RootTransformWeightAttributeId);
+						RootMotionAttribute->Value = Keyframe.Pose.LocalTransformsView[RootIndex];
+						RootMotionWeight->Value = 1.f;
 
-				const int32 RootIndex = Keyframe.Pose.GetRefPose().GetLODBoneIndexFromSkeletonBoneIndex(0);
-				if (RootIndex != INDEX_NONE)
-				{
-					Keyframe.Pose.LocalTransformsView[RootIndex] = FTransform::Identity;
+						Keyframe.Pose.LocalTransformsView[RootIndex] = FTransform::Identity;
+					}
 				}
 			}
 #if WITH_EDITOR
