Refresh Mask Sections when undoing changes. This fixes an issue where undoing mask section deletion did not refresh the ui, preventing the section from being shown.

#Synced-CL 51716622

diff --git a/Engine/Plugins/Experimental/MovieSceneAnimMixer/Source/MovieSceneAnimMixer/Private/MovieSceneAnimationMaskDecoration.cpp b/Engine/Plugins/Experimental/MovieSceneAnimMixer/Source/MovieSceneAnimMixer/Private/MovieSceneAnimationMaskDecoration.cpp
--- a/Engine/Plugins/Experimental/MovieSceneAnimMixer/Source/MovieSceneAnimMixer/Private/MovieSceneAnimationMaskDecoration.cpp
+++ b/Engine/Plugins/Experimental/MovieSceneAnimMixer/Source/MovieSceneAnimMixer/Private/MovieSceneAnimationMaskDecoration.cpp
@@ -213,3 +213,11 @@ void UMovieSceneAnimationMaskDecoration::RemoveSection(UMovieSceneSection& Secti
 	MarkStructureChanged();
 #endif
 }
+
+#if WITH_EDITOR
+void UMovieSceneAnimationMaskDecoration::PostEditUndo()
+{
+	Super::PostEditUndo();
+	MarkStructureChanged();
+}
+#endif
\ No newline at end of file
diff --git a/Engine/Plugins/Experimental/MovieSceneAnimMixer/Source/MovieSceneAnimMixer/Public/MovieSceneAnimationMaskDecoration.h b/Engine/Plugins/Experimental/MovieSceneAnimMixer/Source/MovieSceneAnimMixer/Public/MovieSceneAnimationMaskDecoration.h
--- a/Engine/Plugins/Experimental/MovieSceneAnimMixer/Source/MovieSceneAnimMixer/Public/MovieSceneAnimationMaskDecoration.h
+++ b/Engine/Plugins/Experimental/MovieSceneAnimMixer/Source/MovieSceneAnimMixer/Public/MovieSceneAnimationMaskDecoration.h
@@ -91,6 +91,11 @@ public:
 	UPROPERTY()
 	TArray<TObjectPtr<UMovieSceneSection>> MaskSections;
 
+#if WITH_EDITOR
+	virtual void PostEditUndo() override;
+#endif
+
+
 private:
 	UPROPERTY()
 	FMovieSceneMaskBlendProvider BlendProvider;
