Custom Menu for Mask Decorations

Implemented menu providers for decorations. We will now look for registered menu providers when building entries for a given decoraiton class.
Added a menu provider for mask decorations that creates an asset picker for UAF Blend Masks that matches by skeleton for the bound object.
Updated SequencerHelpers::BuildDecorationMenu to accept an FGuid for the bound object to allow for filtering based on the bound object when building the decoration menu.

#Synced-CL 51759331

diff --git a/Engine/Plugins/Experimental/MovieSceneAnimMixer/Source/MovieSceneAnimMixer/Private/MovieSceneAnimMixerMaskSection.cpp b/Engine/Plugins/Experimental/MovieSceneAnimMixer/Source/MovieSceneAnimMixer/Private/MovieSceneAnimMixerMaskSection.cpp
--- a/Engine/Plugins/Experimental/MovieSceneAnimMixer/Source/MovieSceneAnimMixer/Private/MovieSceneAnimMixerMaskSection.cpp
+++ b/Engine/Plugins/Experimental/MovieSceneAnimMixer/Source/MovieSceneAnimMixer/Private/MovieSceneAnimMixerMaskSection.cpp
@@ -20,6 +20,14 @@ UMovieSceneAnimMixerMaskSection::UMovieSceneAnimMixerMaskSection()
 	SetColorTint(MaskTintColor);
 }
 
+void UMovieSceneAnimMixerMaskSection::SetBlendMask(UObject* InMask)
+{
+	if (UUAFBlendMask* InBlendMask = Cast<UUAFBlendMask>(InMask))
+	{
+		BlendMask = InBlendMask;
+	}
+}
+
 void UMovieSceneAnimMixerMaskSection::ImportEntityImpl(UMovieSceneEntitySystemLinker* EntityLinker, const FEntityImportParams& Params,
 	FImportedEntity* OutImportedEntity)
 {
diff --git a/Engine/Plugins/Experimental/MovieSceneAnimMixer/Source/MovieSceneAnimMixer/Public/MovieSceneAnimMixerMaskSection.h b/Engine/Plugins/Experimental/MovieSceneAnimMixer/Source/MovieSceneAnimMixer/Public/MovieSceneAnimMixerMaskSection.h
--- a/Engine/Plugins/Experimental/MovieSceneAnimMixer/Source/MovieSceneAnimMixer/Public/MovieSceneAnimMixerMaskSection.h
+++ b/Engine/Plugins/Experimental/MovieSceneAnimMixer/Source/MovieSceneAnimMixer/Public/MovieSceneAnimMixerMaskSection.h
@@ -3,7 +3,6 @@
 #pragma once
 #include <UAF/BlendMask/UAFBlendMask.h>
 
-#include "HierarchyTable.h"
 #include "MovieSceneAnimationMixerItemInterface.h"
 #include "MovieSceneAnimationMixerLayer.h"
 #include "MovieSceneSection.h"
@@ -12,6 +11,8 @@
 
 #include "MovieSceneAnimMixerMaskSection.generated.h"
 
+#define UE_API MOVIESCENEANIMMIXER_API
+
 struct FMovieSceneMaskBlendProvider;
 
 UCLASS(MinimalAPI, DisplayName="Mask")
@@ -30,6 +31,8 @@ public:
 
 	TObjectKey<UUAFBlendMask> GetBlendMask() { return BlendMask; }
 
+	UE_API void SetBlendMask(UObject* InMask);
+
 #if WITH_EDITOR
 	virtual bool SupportsVerticalDragging() const override { return false; }
 
@@ -50,3 +53,5 @@ private:
 
 	FColor MaskTintColor = FColor(110, 25, 50);
 };
+
+#undef UE_API
\ No newline at end of file
diff --git a/Engine/Plugins/Experimental/MovieSceneAnimMixer/Source/MovieSceneAnimMixerEditor/Private/AnimMixerDecorationMenuProviders.cpp b/Engine/Plugins/Experimental/MovieSceneAnimMixer/Source/MovieSceneAnimMixerEditor/Private/AnimMixerDecorationMenuProviders.cpp
new file mode 100644
--- /dev/null
+++ b/Engine/Plugins/Experimental/MovieSceneAnimMixer/Source/MovieSceneAnimMixerEditor/Private/AnimMixerDecorationMenuProviders.cpp
@@ -0,0 +1,177 @@
+﻿// Copyright Epic Games, Inc. All Rights Reserved.
+
+#include "AnimMixerDecorationMenuProviders.h"
+
+#include "AnimatedRange.h"
+#include "ContentBrowserModule.h"
+#include "IContentBrowserSingleton.h"
+#include "ISequencer.h"
+#include "MovieScene.h"
+#include "MovieSceneAnimationMaskDecoration.h"
+#include "MovieSceneAnimMixerMaskSection.h"
+#include "ScopedTransaction.h"
+#include "SequencerSettings.h"
+#include "Components/SkeletalMeshComponent.h"
+#include "Decorations/MovieSceneDecorationContainer.h"
+#include "Engine/SkeletalMesh.h"
+#include "Framework/MultiBox/MultiBoxBuilder.h"
+#include "GameFramework/Actor.h"
+#include "UAF/BlendMask/UAFBlendMask.h"
+
+#define LOCTEXT_NAMESPACE "AnimmixerDecorationMenuProviders"
+
+UClass* FMaskDecorationMenuProvider::GetHandledDecorationClass() const
+{
+	return UMovieSceneAnimationMaskDecoration::StaticClass();
+}
+
+void FMaskDecorationMenuProvider::PopulateAddDecorationMenu(FMenuBuilder& MenuBuilder, TObjectPtr<UObject> BoundObject,
+	UMovieSceneDecorationContainerObject* DecorationContainer, TWeakPtr<ISequencer> InSequencer)
+{
+	MenuBuilder.AddSubMenu(
+		LOCTEXT("MaskDecoration", "Mask"),
+		LOCTEXT("MaskDecorationTooltip", "Add a mask from the selected asset."),
+		FNewMenuDelegate::CreateRaw(this, &FMaskDecorationMenuProvider::PopulateMaskSubmenu, BoundObject, DecorationContainer, InSequencer)
+	);
+}
+
+const USkeleton* FMaskDecorationMenuProvider::FindSkeletonForObject(const TObjectPtr<UObject>& Object)
+{
+	const USkeletalMeshComponent* SkeletalMesh = nullptr;
+
+	// If the bound object is an actor, look for a mesh component.
+	if (const AActor* Actor = Cast<AActor>(Object))
+	{
+		TArray<USkeletalMeshComponent*> SkeletalMeshComponents;
+		Actor->GetComponents(SkeletalMeshComponents);
+
+		// If there is only one mesh component on the actor, use that. If there are multiple, it's ambiguous which one to use, and the user should bind to components instead.
+		// This will result in no filtering of masks, which can allow the user to bind an incompatible mask, and may result in unwanted behavior.
+		// Choosing to filter based on the first mesh, regardless of if there are others, could result in the same behavior.
+		// Not filtering at least makes the correct masks visible to all components.
+		if (SkeletalMeshComponents.Num() == 1)
+		{
+			SkeletalMesh = SkeletalMeshComponents[0];
+		}
+	}
+	// If the bound object is a component, use it. 
+	else
+	{
+		SkeletalMesh = Cast<USkeletalMeshComponent>(Object);
+	}
+
+	return SkeletalMesh && SkeletalMesh->GetSkeletalMeshAsset() ? SkeletalMesh->GetSkeletalMeshAsset()->GetSkeleton() : nullptr;
+}
+
+void FMaskDecorationMenuProvider::PopulateMaskSubmenu(
+	FMenuBuilder& MenuBuilder,
+	TObjectPtr<UObject> BoundObject,
+	UMovieSceneDecorationContainerObject* DecorationContainer,
+	TWeakPtr<ISequencer> InSequencer
+	)
+{
+	const USkeleton* Skeleton = FindSkeletonForObject(BoundObject);
+
+	FAssetPickerConfig AssetPickerConfig;
+	{
+		AssetPickerConfig.OnAssetSelected = FOnAssetSelected::CreateRaw(this, &FMaskDecorationMenuProvider::OnMaskAssetSelected, DecorationContainer, InSequencer);
+		AssetPickerConfig.bAllowNullSelection = false;
+		AssetPickerConfig.bAddFilterUI = true;
+		AssetPickerConfig.bShowTypeInColumnView = false;
+		AssetPickerConfig.InitialAssetViewType = EAssetViewType::List;
+		AssetPickerConfig.Filter.bRecursiveClasses = true;
+		AssetPickerConfig.Filter.ClassPaths.Add(UUAFBlendMask::StaticClass()->GetClassPathName());
+		AssetPickerConfig.OnShouldFilterAsset.BindRaw(this, &FMaskDecorationMenuProvider::FilterBlendMasks, Skeleton);
+		AssetPickerConfig.SaveSettingsName = TEXT("SequencerAssetPicker");
+		AssetPickerConfig.AdditionalReferencingAssets.Add(FAssetData());
+	}
+
+	const FContentBrowserModule& ContentBrowserModule = FModuleManager::Get().LoadModuleChecked<FContentBrowserModule>(TEXT("ContentBrowser"));
+
+	const float WidthOverride = InSequencer.IsValid() ? InSequencer.Pin()->GetSequencerSettings()->GetAssetBrowserWidth() : 500.f;
+	const float HeightOverride = InSequencer.IsValid() ? InSequencer.Pin()->GetSequencerSettings()->GetAssetBrowserHeight() : 400.f;
+
+	const TSharedPtr<SBox> MenuEntry = SNew(SBox)
+		.WidthOverride(WidthOverride)
+		.HeightOverride(HeightOverride)
+		[
+			ContentBrowserModule.Get().CreateAssetPicker(AssetPickerConfig)
+		];
+
+	MenuBuilder.AddWidget(MenuEntry.ToSharedRef(), FText::GetEmpty(), true);
+}
+
+void FMaskDecorationMenuProvider::OnMaskAssetSelected(const FAssetData& AssetData, UMovieSceneDecorationContainerObject* DecorationContainer, TWeakPtr<ISequencer> WeakSequencer)
+{
+	if (!DecorationContainer)
+	{
+		return;
+	}
+
+	const TSharedPtr<ISequencer> Sequencer = WeakSequencer.Pin();
+	if (!Sequencer)
+	{
+		return;
+	}
+
+	FScopedTransaction Transaction(FText::Format(LOCTEXT("AddDecorationSection", "Add {0}"), UMovieSceneAnimationMaskDecoration::StaticClass()->GetDisplayNameText()));
+
+	UMovieSceneAnimationMaskDecoration* MaskDecoration = DecorationContainer->GetOrCreateDecoration<UMovieSceneAnimationMaskDecoration>();
+
+	// The newly added section will start at the current sequence time, and by default end at the playback end.
+	const FQualifiedFrameTime CurrentTime = Sequencer->GetLocalTime();
+	const FFrameNumber PlaybackEnd = UE::MovieScene::DiscreteExclusiveUpper(Sequencer->GetFocusedMovieSceneSequence()->GetMovieScene()->GetPlaybackRange());
+
+	// Check to see if the playback end is a valid end time, and if not, use the view range of the sequence instead.
+	FFrameNumber NewSectionRangeEnd = PlaybackEnd;
+	if (PlaybackEnd <= CurrentTime.Time.FrameNumber)
+	{
+		const FAnimatedRange ViewRange = Sequencer->GetViewRange();
+		const FFrameRate TickResolution = Sequencer->GetFocusedTickResolution();
+		NewSectionRangeEnd = (ViewRange.GetUpperBoundValue() * TickResolution).FloorToFrame();
+	}
+
+	// Create the section, and modify the container object.
+	DecorationContainer->Modify();
+	UMovieSceneSection* NewSection = MaskDecoration->CreateNewSection();
+
+	// Set the range. Masks don't support infinite ranges, so always set it's range.
+	NewSection->SetRange(TRange<FFrameNumber>(CurrentTime.Time.FrameNumber, NewSectionRangeEnd));
+
+	// Add the section, this will fix up existing sections to avoid overlaps, and may change the range originally set above.
+	MaskDecoration->AddSection(NewSection);
+
+	// Set the blend mask to the selected asset
+	UObject* SelectedMask = AssetData.GetAsset();
+	if (UMovieSceneAnimMixerMaskSection* MaskSection = Cast<UMovieSceneAnimMixerMaskSection>(NewSection))
+	{
+		MaskSection->SetBlendMask(SelectedMask);
+	}
+	DecorationContainer->MarkAsChanged();
+}
+
+bool FMaskDecorationMenuProvider::FilterBlendMasks(const FAssetData& AssetData, const USkeleton* Skeleton)
+{
+	// Unlike animation sequences this function does not filter using USkeleton::IsCompatibleForEditor since that checks for compatible skeletons,
+	// which is not relevant to blend masks.
+	// Instead, this function will just check that their export text names match.
+
+	// If there is no skeleton, don't filter blend masks. Default to showing all masks.
+	if (!Skeleton)
+	{
+		return false;
+	}
+	
+	// Get the value of the blend mask's skeleton tag
+	const FString AssetSkeleton = AssetData.GetTagValueRef<FString>(TEXT("Skeleton"));
+
+	// Get the export text of the skeleton associated with the owning bound object.
+	TStringBuilder<128> SkeletonStringBuilder;
+	FAssetData(Skeleton, FAssetData::ECreationFlags::SkipAssetRegistryTagsGathering).GetExportTextName(SkeletonStringBuilder);
+	const TCHAR* SkeletonString = SkeletonStringBuilder.ToString();
+	
+	// Filter (do not show) the blend mask asset if the skeletons do not match.
+	return  AssetSkeleton != SkeletonString;
+}
+
+#undef LOCTEXT_NAMESPACE
\ No newline at end of file
diff --git a/Engine/Plugins/Experimental/MovieSceneAnimMixer/Source/MovieSceneAnimMixerEditor/Private/AnimMixerDecorationMenuProviders.h b/Engine/Plugins/Experimental/MovieSceneAnimMixer/Source/MovieSceneAnimMixerEditor/Private/AnimMixerDecorationMenuProviders.h
new file mode 100644
--- /dev/null
+++ b/Engine/Plugins/Experimental/MovieSceneAnimMixer/Source/MovieSceneAnimMixerEditor/Private/AnimMixerDecorationMenuProviders.h
@@ -0,0 +1,32 @@
+﻿// Copyright Epic Games, Inc. All Rights Reserved.
+
+#pragma once
+#include "Animation/Skeleton.h"
+#include "AssetRegistry/AssetData.h"
+#include "Decorations/IMovieSceneDecorationMenuProvider.h"
+
+/**
+ * Decoration menu provider for masks (UMovieSceneAnimationMaskDecoration).
+ * Adds a "Mask" menu entry that allows the user to select a blend mask asset, which will create a mask section with that asset as its blend mask.
+ */
+class FMaskDecorationMenuProvider : public IMovieSceneDecorationMenuProvider
+{
+public:
+
+	/** IMovieSceneDecorationMenuProvider **/
+	virtual UClass* GetHandledDecorationClass() const override;
+	virtual void PopulateAddDecorationMenu(FMenuBuilder& MenuBuilder, TObjectPtr<UObject> BoundObject, UMovieSceneDecorationContainerObject* DecorationContainer, TWeakPtr<ISequencer> InSequencer) override;
+	// Masks are high priority.
+	virtual int32 GetDecorationMenuPriority() const override { return  0; };
+
+private:
+
+	const USkeleton* FindSkeletonForObject(const TObjectPtr<UObject>& Object);
+
+	void PopulateMaskSubmenu(FMenuBuilder& MenuBuilder, TObjectPtr<UObject> BoundObject, UMovieSceneDecorationContainerObject* DecorationContainer, TWeakPtr<ISequencer> InSequencer);
+
+	void OnMaskAssetSelected(const FAssetData& AssetData, UMovieSceneDecorationContainerObject* DecorationContainer, TWeakPtr<ISequencer> Sequencer);
+
+	bool FilterBlendMasks(const FAssetData& AssetData, const USkeleton* Skeleton);
+
+};
\ No newline at end of file
diff --git a/Engine/Plugins/Experimental/MovieSceneAnimMixer/Source/MovieSceneAnimMixerEditor/Private/MovieSceneAnimMixerEditorModule.cpp b/Engine/Plugins/Experimental/MovieSceneAnimMixer/Source/MovieSceneAnimMixerEditor/Private/MovieSceneAnimMixerEditorModule.cpp
--- a/Engine/Plugins/Experimental/MovieSceneAnimMixer/Source/MovieSceneAnimMixerEditor/Private/MovieSceneAnimMixerEditorModule.cpp
+++ b/Engine/Plugins/Experimental/MovieSceneAnimMixer/Source/MovieSceneAnimMixerEditor/Private/MovieSceneAnimMixerEditorModule.cpp
@@ -3,6 +3,7 @@
 #include "MovieSceneAnimMixerEditorModule.h"
 
 #include "AnimBlueprintExtension_SequencerMixerTarget.h"
+#include "AnimMixerDecorationMenuProviders.h"
 #include "AnimMixerTargetMenuProviders.h"
 #include "AnimTransitionSectionInterface.h"
 #include "IAnimGraph_SequencerMixerTargetConnector.h"
@@ -217,6 +218,10 @@ namespace UE::MovieScene
 
 		AnimMaskingDecorationEditor = MakeUnique<FMovieSceneAnimationMaskDecorationEditor>();
 		IModularFeatures::Get().RegisterModularFeature(ISequencerDecorationEditor::GetModularFeatureName(), AnimMaskingDecorationEditor.Get());
+
+		// Register decoration menu providers
+		MaskDecorationMenuProvider = MakeUnique<FMaskDecorationMenuProvider>();
+		IModularFeatures::Get().RegisterModularFeature(IMovieSceneDecorationMenuProvider::GetModularFeatureName(), MaskDecorationMenuProvider.Get());
 	}
 
 	void FMovieSceneAnimMixerEditorModule::ShutdownModule()
@@ -236,6 +241,12 @@ namespace UE::MovieScene
 			IModularFeatures::Get().UnregisterModularFeature(ISequencerDecorationEditor::GetModularFeatureName(), AnimMaskingDecorationEditor.Get());
 		}
 
+		// Unregister decoration menu providers
+		if (MaskDecorationMenuProvider)
+		{
+			IModularFeatures::Get().UnregisterModularFeature(IMovieSceneDecorationMenuProvider::GetModularFeatureName(), MaskDecorationMenuProvider.Get());
+		}
+
 		// Unregister target menu providers
 		if (AutomaticTargetMenuProvider)
 		{
diff --git a/Engine/Plugins/Experimental/MovieSceneAnimMixer/Source/MovieSceneAnimMixerEditor/Private/MovieSceneAnimMixerEditorModule.h b/Engine/Plugins/Experimental/MovieSceneAnimMixer/Source/MovieSceneAnimMixerEditor/Private/MovieSceneAnimMixerEditorModule.h
--- a/Engine/Plugins/Experimental/MovieSceneAnimMixer/Source/MovieSceneAnimMixerEditor/Private/MovieSceneAnimMixerEditorModule.h
+++ b/Engine/Plugins/Experimental/MovieSceneAnimMixer/Source/MovieSceneAnimMixerEditor/Private/MovieSceneAnimMixerEditorModule.h
@@ -11,6 +11,7 @@ class FAutomaticTargetMenuProvider;
 class FAnimInstanceTargetMenuProvider;
 class FAnimBlueprintTargetMenuProvider;
 class FAnimNextInjectionTargetMenuProvider;
+class FMaskDecorationMenuProvider;
 
 namespace UE::Sequencer
 {
@@ -44,6 +45,9 @@ namespace UE::MovieScene
 		TUniquePtr<UE::Sequencer::FRootMotionTargetDecorationEditor> RootMotionTargetDecorationEditor;
 		TUniquePtr<UE::Sequencer::FRootMotionSettingsDecorationEditor> RootMotionSettingsDecorationEditor;
 		TUniquePtr<Sequencer::FMovieSceneAnimationMaskDecorationEditor> AnimMaskingDecorationEditor;
+
+		// Decoration menu providers
+		TUniquePtr<FMaskDecorationMenuProvider> MaskDecorationMenuProvider;
 	};
 }
 
diff --git a/Engine/Plugins/Experimental/MovieSceneAnimMixer/Source/MovieSceneAnimMixerEditor/Private/MovieSceneAnimationMixerTrackEditor.cpp b/Engine/Plugins/Experimental/MovieSceneAnimMixer/Source/MovieSceneAnimMixerEditor/Private/MovieSceneAnimationMixerTrackEditor.cpp
--- a/Engine/Plugins/Experimental/MovieSceneAnimMixer/Source/MovieSceneAnimMixerEditor/Private/MovieSceneAnimationMixerTrackEditor.cpp
+++ b/Engine/Plugins/Experimental/MovieSceneAnimMixer/Source/MovieSceneAnimMixerEditor/Private/MovieSceneAnimationMixerTrackEditor.cpp
@@ -1096,20 +1096,20 @@ TSharedRef<SWidget> FAnimationMixerTrackEditor::BuildAddSectionSubMenu(TWeakView
 		UMovieSceneAnimationMixerLayer* Layer = MixerTrack->GetLayer(LayerExtension->GetRowIndex());
 		if (Layer)
 		{
-			SequencerHelpers::BuildDecorationMenu(MenuBuilder, Layer, Sequencer);
+			SequencerHelpers::BuildDecorationMenu(MenuBuilder, Layer, ObjectBindingGuid, Sequencer);
 
 			// For single-section layers, also show decorations compatible with the section
 			TArray<UMovieSceneSection*> LayerSections = Layer->GetSections();
 			if (LayerSections.Num() == 1 && LayerSections[0])
 			{
-				SequencerHelpers::BuildDecorationMenu(MenuBuilder, LayerSections[0], Sequencer);
+				SequencerHelpers::BuildDecorationMenu(MenuBuilder, LayerSections[0], ObjectBindingGuid, Sequencer);
 			}
 		}
 	}
 	else
 	{
 		// Track-level: add decorations to the track
-		SequencerHelpers::BuildDecorationMenu(MenuBuilder, MixerTrack, Sequencer);
+		SequencerHelpers::BuildDecorationMenu(MenuBuilder, MixerTrack, ObjectBindingGuid, Sequencer);
 	}
 
 	return MenuBuilder.MakeWidget();
diff --git a/Engine/Source/Editor/MovieSceneTools/Private/Decorations/IMovieSceneDecorationMenuProvider.cpp b/Engine/Source/Editor/MovieSceneTools/Private/Decorations/IMovieSceneDecorationMenuProvider.cpp
new file mode 100644
--- /dev/null
+++ b/Engine/Source/Editor/MovieSceneTools/Private/Decorations/IMovieSceneDecorationMenuProvider.cpp
@@ -0,0 +1,13 @@
+﻿// Copyright Epic Games, Inc. All Rights Reserved.
+
+#include "Decorations/IMovieSceneDecorationMenuProvider.h"
+
+FName IMovieSceneDecorationMenuProvider::GetModularFeatureName()
+{
+	static FName FeatureName = FName(TEXT("MovieSceneDecorationProvider"));
+	return FeatureName;
+}
+
+IMovieSceneDecorationMenuProvider::~IMovieSceneDecorationMenuProvider()
+{
+}
diff --git a/Engine/Source/Editor/MovieSceneTools/Public/Decorations/IMovieSceneDecorationMenuProvider.h b/Engine/Source/Editor/MovieSceneTools/Public/Decorations/IMovieSceneDecorationMenuProvider.h
new file mode 100644
--- /dev/null
+++ b/Engine/Source/Editor/MovieSceneTools/Public/Decorations/IMovieSceneDecorationMenuProvider.h
@@ -0,0 +1,46 @@
+﻿// Copyright Epic Games, Inc. All Rights Reserved.
+
+#pragma once
+
+#include "CoreMinimal.h"
+#include "Features/IModularFeature.h"
+
+#define UE_API MOVIESCENETOOLS_API
+
+class ISequencer;
+class UMovieSceneDecorationContainerObject;
+class FMenuBuilder;
+
+class IMovieSceneDecorationMenuProvider : public IModularFeature
+{
+public:
+
+	static UE_API FName GetModularFeatureName();
+
+	UE_API virtual ~IMovieSceneDecorationMenuProvider();
+
+	/**
+	 * Allows providers to specify the decoration classes they handle.
+	 * @return The decoration UClass this menu provider handles.
+	 */
+	virtual const UClass* GetHandledDecorationClass() const = 0;
+
+
+	/**
+	 * Populate the custom menu entry for the decoration class handled by this provider.
+	 * @param MenuBuilder The menu builder to extend, and add a custom entry to.
+	 * @param BoundObject The object associated with the decoration container, if it exists.
+	 * @param DecorationContainer The decoration container owning the decoration this provider is associated with.
+	 * @param InSequencer The sequence editor.
+	 */
+	UE_API virtual void PopulateAddDecorationMenu(FMenuBuilder& MenuBuilder, TObjectPtr<UObject> BoundObject, UMovieSceneDecorationContainerObject* DecorationContainer, TWeakPtr<ISequencer> InSequencer) = 0;
+
+	/**
+	 * Allows providers to define their priority in their menu entries. 
+	 * @return Menu priority for this decoration provider. Lower values appear first in the menu.
+	 * Default is 100 (middle priority). Use 0-50 for high priority items, 100+ for normal items.
+	 */
+	virtual int32 GetDecorationMenuPriority() const { return 100; };
+};
+
+#undef UE_API
\ No newline at end of file
diff --git a/Engine/Source/Editor/Sequencer/Private/MVVM/ViewModels/DecorationOutlinerModel.cpp b/Engine/Source/Editor/Sequencer/Private/MVVM/ViewModels/DecorationOutlinerModel.cpp
--- a/Engine/Source/Editor/Sequencer/Private/MVVM/ViewModels/DecorationOutlinerModel.cpp
+++ b/Engine/Source/Editor/Sequencer/Private/MVVM/ViewModels/DecorationOutlinerModel.cpp
@@ -324,10 +324,12 @@ TSharedRef<SWidget> FDecorationOutlinerModel::BuildAddDecorationSubMenu(TSharedP
 		return MenuBuilder.MakeWidget();
 	}
 
+	FGuid ObjectBindingID = WeakParentTrack.IsValid() ? WeakParentTrack->FindObjectBindingGuid() : FGuid();
+
 	if (Decoration)
 	{
 		UMovieSceneDecorationContainerObject* DecorationContainer = Decoration->GetTypedOuter<UMovieSceneDecorationContainerObject>();
-		SequencerHelpers::BuildDecorationMenu(MenuBuilder, DecorationContainer, SequencerEditor->GetSequencer());
+		SequencerHelpers::BuildDecorationMenu(MenuBuilder, DecorationContainer, ObjectBindingID, SequencerEditor->GetSequencer());
 	}
 
 	return  MenuBuilder.MakeWidget();
diff --git a/Engine/Source/Editor/Sequencer/Private/MVVM/ViewModels/SectionOutlinerModel.cpp b/Engine/Source/Editor/Sequencer/Private/MVVM/ViewModels/SectionOutlinerModel.cpp
--- a/Engine/Source/Editor/Sequencer/Private/MVVM/ViewModels/SectionOutlinerModel.cpp
+++ b/Engine/Source/Editor/Sequencer/Private/MVVM/ViewModels/SectionOutlinerModel.cpp
@@ -298,13 +298,16 @@ TSharedPtr<SWidget> FSectionOutlinerModel::CreateOutlinerViewForColumn(const FCr
 		UMovieSceneSection* Section = GetSection();
 		if (Section && SequencerHelpers::HasCompatibleDecorations(Section))
 		{
+			TViewModelPtr<IObjectBindingExtension> ObjectBinding = FindAncestorOfType<IObjectBindingExtension>();
+			FGuid ObjectBindingID = ObjectBinding ? ObjectBinding->GetObjectGuid() : FGuid();
+
 			TWeakObjectPtr<UMovieSceneSection> WeakSectionForMenu = Section;
-			auto GetMenuContent = [WeakSectionForMenu, Editor]() -> TSharedRef<SWidget>
+			auto GetMenuContent = [WeakSectionForMenu, Editor, ObjectBindingID]() -> TSharedRef<SWidget>
 			{
 				FMenuBuilder MenuBuilder(true, nullptr);
 				if (UMovieSceneSection* PinnedSection = WeakSectionForMenu.Get())
 				{
-					SequencerHelpers::BuildDecorationMenu(MenuBuilder, PinnedSection, Editor->GetSequencer());
+					SequencerHelpers::BuildDecorationMenu(MenuBuilder, PinnedSection, ObjectBindingID, Editor->GetSequencer());
 				}
 				return MenuBuilder.MakeWidget();
 			};
diff --git a/Engine/Source/Editor/Sequencer/Private/MovieSceneTrackEditor.cpp b/Engine/Source/Editor/Sequencer/Private/MovieSceneTrackEditor.cpp
--- a/Engine/Source/Editor/Sequencer/Private/MovieSceneTrackEditor.cpp
+++ b/Engine/Source/Editor/Sequencer/Private/MovieSceneTrackEditor.cpp
@@ -364,7 +364,7 @@ TSharedPtr<SWidget> FMovieSceneTrackEditor::BuildOutlinerEditWidget(const FGuid&
 			// Add decoration entries if track supports decorations
 			if (bHasDecorations)
 			{
-				SequencerHelpers::BuildDecorationMenu(MenuBuilder, Track, WeakSequencer);
+				SequencerHelpers::BuildDecorationMenu(MenuBuilder, Track, Track->FindObjectBindingGuid(), WeakSequencer);
 			}
 
 			return MenuBuilder.MakeWidget();
diff --git a/Engine/Source/Editor/Sequencer/Private/SequencerCommonHelpers.cpp b/Engine/Source/Editor/Sequencer/Private/SequencerCommonHelpers.cpp
--- a/Engine/Source/Editor/Sequencer/Private/SequencerCommonHelpers.cpp
+++ b/Engine/Source/Editor/Sequencer/Private/SequencerCommonHelpers.cpp
@@ -46,6 +46,8 @@
 #include "Styling/CoreStyle.h"
 #include "Conditions/MovieSceneDirectorBlueprintConditionCustomization.h"
 #include "Conditions/MovieSceneConditionCustomization.h"
+#include "Decorations/IMovieSceneDecorationMenuProvider.h"
+#include "Features/IModularFeatures.h"
 
 void SequencerHelpers::GetAllChannels(TSharedPtr<FViewModel> DataModel, TSet<TSharedPtr<UE::Sequencer::FChannelModel>>& Channels)
 {
@@ -1192,7 +1194,9 @@ bool SequencerHelpers::HasOutlinerVisibleDecorations(const UMovieSceneDecoration
 
 void SequencerHelpers::BuildDecorationMenu(
 	FMenuBuilder& MenuBuilder,
-	UMovieSceneDecorationContainerObject* DecorationContainer, TWeakPtr<ISequencer> InSequencer)
+	UMovieSceneDecorationContainerObject* DecorationContainer,
+	const FGuid& ObjectBindingID,
+	TWeakPtr<ISequencer> InSequencer)
 {
 	if (!DecorationContainer)
 	{
@@ -1207,8 +1211,46 @@ void SequencerHelpers::BuildDecorationMenu(
 		return;
 	}
 
+	UObject* BoundObject = nullptr;
+
+	if (const TSharedPtr<ISequencer> PinnedSequencer = InSequencer.Pin())
+	{
+		const TArrayView<TWeakObjectPtr<>> Objects = PinnedSequencer->FindObjectsInCurrentSequence(ObjectBindingID);
+		if (Objects.Num())
+		{
+			BoundObject = Objects[0].Get();
+		}
+	}
+
+	// Get menu providers
+	TArray<IMovieSceneDecorationMenuProvider*> MenuProviders = IModularFeatures::Get().GetModularFeatureImplementations<IMovieSceneDecorationMenuProvider>(IMovieSceneDecorationMenuProvider::GetModularFeatureName());
+
+	// Sort the menu providers so their entries are added in the order specified by their priority.
+	MenuProviders.Sort([](const IMovieSceneDecorationMenuProvider& A, const IMovieSceneDecorationMenuProvider& B)
+	{
+		return A.GetDecorationMenuPriority() < B.GetDecorationMenuPriority();
+	});
+
 	MenuBuilder.BeginSection(NAME_None, NSLOCTEXT("SequencerHelpers", "Modifiers_Label", "Modifiers"));
 	{
+
+		// First, allow menu providers to handle decorations.
+		for (IMovieSceneDecorationMenuProvider* MenuProvider : MenuProviders)
+		{
+			if (!MenuProvider)
+			{
+				continue;
+			}
+			const UClass* HandledClass = MenuProvider->GetHandledDecorationClass();
+			if (CompatibleDecorations.Contains(HandledClass))
+			{
+				// The custom menu provider should handle the case when the decoration already exists, so PopulateAddDecorationMenu is always called.
+				MenuProvider->PopulateAddDecorationMenu(MenuBuilder, BoundObject, DecorationContainer, InSequencer);
+				// Remove class, so its entry isn't duplicated. 
+				CompatibleDecorations.Remove(HandledClass);
+			}
+		}
+
 		for (UClass* DecorationClass : CompatibleDecorations)
 		{
 			const TObjectPtr<UObject>* ExistingDecoration = DecorationContainer->GetDecorations().FindByPredicate(
diff --git a/Engine/Source/Editor/Sequencer/Public/SequencerCommonHelpers.h b/Engine/Source/Editor/Sequencer/Public/SequencerCommonHelpers.h
--- a/Engine/Source/Editor/Sequencer/Public/SequencerCommonHelpers.h
+++ b/Engine/Source/Editor/Sequencer/Public/SequencerCommonHelpers.h
@@ -160,6 +160,7 @@ public:
 	 */
 	static void BuildDecorationMenu(FMenuBuilder& MenuBuilder
 		, UMovieSceneDecorationContainerObject* DecorationContainer
+		, const FGuid& ObjectBindingID
 		, TWeakPtr<ISequencer> InSequencer);
 
 	/*
