Sequencer: Fix issue with the anim mixer bake's poses- these were stack based and need to be heap based so they can be accessed after the mem stack allocator goes away.

#Synced-CL 51609808

diff --git a/Engine/Plugins/Experimental/MovieSceneAnimMixer/Source/MovieSceneAnimMixer/Private/AnimMixerBakeEvaluation.cpp b/Engine/Plugins/Experimental/MovieSceneAnimMixer/Source/MovieSceneAnimMixer/Private/AnimMixerBakeEvaluation.cpp
--- a/Engine/Plugins/Experimental/MovieSceneAnimMixer/Source/MovieSceneAnimMixer/Private/AnimMixerBakeEvaluation.cpp
+++ b/Engine/Plugins/Experimental/MovieSceneAnimMixer/Source/MovieSceneAnimMixer/Private/AnimMixerBakeEvaluation.cpp
@@ -130,7 +130,9 @@ struct FBakeEvaluationScope
 		{
 			UMovieSceneMixerBakeTargetSystem::FBakedResult& Baked = BakeSystem->BakedResults.Last();
 			Result.RootMotionTransform = Baked.RootMotionTransform;
-			Result.Keyframe = MoveTemp(Baked.Keyframe);
+			Result.Pose = MoveTemp(Baked.Pose);
+			Result.Curves = MoveTemp(Baked.Curves);
+			Result.Attributes = MoveTemp(Baked.Attributes);
 		}
 
 		// When capture is active, the conversion task wrote the pre-conversion
diff --git a/Engine/Plugins/Experimental/MovieSceneAnimMixer/Source/MovieSceneAnimMixer/Private/AnimMixerBoneMatching.cpp b/Engine/Plugins/Experimental/MovieSceneAnimMixer/Source/MovieSceneAnimMixer/Private/AnimMixerBoneMatching.cpp
--- a/Engine/Plugins/Experimental/MovieSceneAnimMixer/Source/MovieSceneAnimMixer/Private/AnimMixerBoneMatching.cpp
+++ b/Engine/Plugins/Experimental/MovieSceneAnimMixer/Source/MovieSceneAnimMixer/Private/AnimMixerBoneMatching.cpp
@@ -168,30 +168,28 @@ static FFrameTime ResolveMatchTime(
 }
 
 FTransform ExtractBoneComponentSpaceTransform(
-	const UE::UAF::FKeyframeState* Keyframe,
+	const UE::UAF::FLODPose& Pose,
 	FName BoneName)
 {
-	if (!Keyframe || !Keyframe->IsValid() || BoneName.IsNone())
+	if (!Pose.IsValid() || BoneName.IsNone())
 	{
 		return FTransform::Identity;
 	}
 
-	FBoneIndexType BoneIndex = Keyframe->Pose.FindLODBoneIndexFromBoneName(BoneName);
+	FBoneIndexType BoneIndex = Pose.FindLODBoneIndexFromBoneName(BoneName);
 	if (BoneIndex == INVALID_BONE_INDEX)
 	{
 		return FTransform::Identity;
 	}
 
-	// Walk up the parent chain, accumulating local transforms into component space.
-	// CS(bone) = Local(bone) * Local(parent) * Local(grandparent) * ... * Local(root)
-	FTransform ComponentSpace = (FTransform)Keyframe->Pose.LocalTransformsView[BoneIndex];
+	FTransform ComponentSpace = (FTransform)Pose.LocalTransformsView[BoneIndex];
 
-	FBoneIndexType ParentIndex = Keyframe->Pose.GetLODBoneParentIndex(BoneIndex);
+	FBoneIndexType ParentIndex = Pose.GetLODBoneParentIndex(BoneIndex);
 	while (ParentIndex != INVALID_BONE_INDEX)
 	{
-		FTransform ParentLocal = (FTransform)Keyframe->Pose.LocalTransformsView[ParentIndex];
+		FTransform ParentLocal = (FTransform)Pose.LocalTransformsView[ParentIndex];
 		ComponentSpace = ComponentSpace * ParentLocal;
-		ParentIndex = Keyframe->Pose.GetLODBoneParentIndex(ParentIndex);
+		ParentIndex = Pose.GetLODBoneParentIndex(ParentIndex);
 	}
 
 	return ComponentSpace;
@@ -234,7 +232,7 @@ FMovieSceneBoneMatchData ComputeBoneMatch(
 	// as RootMotionTransform. ExtractBoneComponentSpaceTransform walks from
 	// the bone to root (which is Identity), giving the bone's transform
 	// relative to root. Composing with RootMotionTransform gives full Component Space.
-	FTransform UnderlyingBoneRelToRoot = ExtractBoneComponentSpaceTransform(UnderlyingResult.Keyframe.Get(), Settings.BoneName);
+	FTransform UnderlyingBoneRelToRoot = ExtractBoneComponentSpaceTransform(UnderlyingResult.Pose, Settings.BoneName);
 	UnderlyingBoneRelToRoot.NormalizeRotation();
 	FTransform UnderlyingRootCS = UnderlyingResult.RootMotionTransform;
 	UnderlyingRootCS.NormalizeRotation();
@@ -289,13 +287,13 @@ FMovieSceneBoneMatchData ComputeBoneMatch(
 			Context.Linker, Context.InstanceHandle, Context.MixerTrack,
 			MatchTime, SectionFilter);
 
-	FTransform SectionBoneRelToRoot = ExtractBoneComponentSpaceTransform(SectionResult.Keyframe.Get(), Settings.BoneName);
+	FTransform SectionBoneRelToRoot = ExtractBoneComponentSpaceTransform(SectionResult.Pose, Settings.BoneName);
 	SectionBoneRelToRoot.NormalizeRotation();
 	FTransform SectionRootCS = SectionResult.RootMotionTransform;
 	SectionRootCS.NormalizeRotation();
 	FTransform SectionBoneCS = SectionBoneRelToRoot * SectionRootCS;
 
-	if (!UnderlyingResult.Keyframe || !SectionResult.Keyframe)
+	if (!UnderlyingResult.IsValid() || !SectionResult.IsValid())
 	{
 		return Result;
 	}
diff --git a/Engine/Plugins/Experimental/MovieSceneAnimMixer/Source/MovieSceneAnimMixer/Private/Systems/MovieSceneMixerBakeTargetSystem.cpp b/Engine/Plugins/Experimental/MovieSceneAnimMixer/Source/MovieSceneAnimMixer/Private/Systems/MovieSceneMixerBakeTargetSystem.cpp
--- a/Engine/Plugins/Experimental/MovieSceneAnimMixer/Source/MovieSceneAnimMixer/Private/Systems/MovieSceneMixerBakeTargetSystem.cpp
+++ b/Engine/Plugins/Experimental/MovieSceneAnimMixer/Source/MovieSceneAnimMixer/Private/Systems/MovieSceneMixerBakeTargetSystem.cpp
@@ -87,7 +87,6 @@ struct FCaptureBakeResults
 		TUniquePtr<FKeyframeState> ResultKeyframe;
 		if (TempVM.PopValue(KEYFRAME_STACK_NAME, ResultKeyframe))
 		{
-			// Read root motion from the RootTransformAttribute on the pose
 			const FAnimMixerComponentTypes* AnimMixerComponents = FAnimMixerComponentTypes::Get();
 			const FTransformAnimationAttribute* RootAttr =
 				ResultKeyframe->Attributes.Find<FTransformAnimationAttribute>(
@@ -97,7 +96,12 @@ struct FCaptureBakeResults
 				Result.RootMotionTransform = RootAttr->Value;
 			}
 
-			Result.Keyframe = MoveTemp(ResultKeyframe);
+			// Copy from stack-allocated keyframe to heap-backed storage
+			// before FMemMark frees the evaluation mem stack.
+			Result.Pose.PrepareForLOD(RefPose, LODLevel);
+			Result.Pose.CopyFrom(ResultKeyframe->Pose);
+			Result.Curves.CopyFrom(ResultKeyframe->Curves);
+			Result.Attributes.CopyFrom(ResultKeyframe->Attributes);
 		}
 
 		BakeSystem->BakedResults.Add(MoveTemp(Result));
diff --git a/Engine/Plugins/Experimental/MovieSceneAnimMixer/Source/MovieSceneAnimMixer/Public/AnimMixerBakeEvaluation.h b/Engine/Plugins/Experimental/MovieSceneAnimMixer/Source/MovieSceneAnimMixer/Public/AnimMixerBakeEvaluation.h
--- a/Engine/Plugins/Experimental/MovieSceneAnimMixer/Source/MovieSceneAnimMixer/Public/AnimMixerBakeEvaluation.h
+++ b/Engine/Plugins/Experimental/MovieSceneAnimMixer/Source/MovieSceneAnimMixer/Public/AnimMixerBakeEvaluation.h
@@ -3,6 +3,9 @@
 #pragma once
 
 #include "CoreTypes.h"
+#include "Animation/AnimCurveTypes.h"
+#include "Animation/AttributesRuntime.h"
+#include "LODPose.h"
 #include "Misc/FrameTime.h"
 #include "UObject/ObjectKey.h"
 
@@ -12,8 +15,6 @@ class USkeletalMeshComponent;
 
 namespace UE::MovieScene { struct FInstanceHandle; }
 
-#include "EvaluationVM/KeyframeState.h"
-
 namespace UE::MovieScene::AnimMixerBakeEvaluation
 {
 
@@ -44,7 +45,11 @@ struct FBakeResult
 	FFrameTime Time;
 	FTransform RootMotionTransform = FTransform::Identity;
 	FTransform AnimationSpaceRootMotion = FTransform::Identity;
-	TUniquePtr<UE::UAF::FKeyframeState> Keyframe;
+	UE::UAF::FLODPoseHeap Pose;
+	FBlendedHeapCurve Curves;
+	UE::Anim::FHeapAttributeContainer Attributes;
+
+	bool IsValid() const { return Pose.IsValid(); }
 };
 
 // Evaluate the mixer at a single time, capturing results through the bake target system.
@@ -60,7 +65,7 @@ MOVIESCENEANIMMIXER_API FBakeResult EvaluateAtTime(
 // Evaluate the mixer at regular intervals from StartTime (inclusive) for NumSamples.
 // More efficient than calling EvaluateAtTime in a loop because save/restore only
 // happens once around the entire batch.
-TArray<FBakeResult> EvaluateRange(
+MOVIESCENEANIMMIXER_API TArray<FBakeResult> EvaluateRange(
 	UMovieSceneEntitySystemLinker* Linker,
 	FInstanceHandle InstanceHandle,
 	UMovieSceneAnimationMixerTrack* MixerTrack,
diff --git a/Engine/Plugins/Experimental/MovieSceneAnimMixer/Source/MovieSceneAnimMixer/Public/Systems/MovieSceneMixerBakeTargetSystem.h b/Engine/Plugins/Experimental/MovieSceneAnimMixer/Source/MovieSceneAnimMixer/Public/Systems/MovieSceneMixerBakeTargetSystem.h
--- a/Engine/Plugins/Experimental/MovieSceneAnimMixer/Source/MovieSceneAnimMixer/Public/Systems/MovieSceneMixerBakeTargetSystem.h
+++ b/Engine/Plugins/Experimental/MovieSceneAnimMixer/Source/MovieSceneAnimMixer/Public/Systems/MovieSceneMixerBakeTargetSystem.h
@@ -5,6 +5,9 @@
 #include "CoreTypes.h"
 #include "AnimMixerBakeEvaluation.h"
 #include "EntitySystem/MovieSceneEntitySystem.h"
+#include "Animation/AnimCurveTypes.h"
+#include "Animation/AttributesRuntime.h"
+#include "LODPose.h"
 #include "MovieSceneMixedAnimationTarget.h"
 #include "MovieSceneMixerBakeTargetSystem.generated.h"
 
@@ -13,7 +16,6 @@ struct FAnimNextEvaluationTask;
 namespace UE::UAF
 {
 	struct FEvaluationVM;
-	struct FKeyframeState;
 }
 
 // Hidden target struct used during bake evaluation to tell the target system
@@ -43,8 +45,12 @@ public:
 	// matches FMovieSceneMixerBakeTarget during a flush. 
 	struct FBakedResult
 	{
-		TUniquePtr<UE::UAF::FKeyframeState> Keyframe;
+		UE::UAF::FLODPoseHeap Pose;
+		FBlendedHeapCurve Curves;
+		UE::Anim::FHeapAttributeContainer Attributes;
 		FTransform RootMotionTransform = FTransform::Identity;
+
+		bool IsValid() const { return Pose.IsValid(); }
 	};
 
 	// Accumulated baked results, populated during bake evaluation
