Sequencer: Add sub-group support to decoration outliner channels, and use that to make the 'Offset' XYZ Yaw Pitch Roll be under an 'Offset' group.

#Synced-CL 51609808

diff --git a/Engine/Plugins/Experimental/MovieSceneAnimMixer/Source/MovieSceneAnimMixer/Private/MovieSceneAnimationMixerTrack.cpp b/Engine/Plugins/Experimental/MovieSceneAnimMixer/Source/MovieSceneAnimMixer/Private/MovieSceneAnimationMixerTrack.cpp
--- a/Engine/Plugins/Experimental/MovieSceneAnimMixer/Source/MovieSceneAnimMixer/Private/MovieSceneAnimationMixerTrack.cpp
+++ b/Engine/Plugins/Experimental/MovieSceneAnimMixer/Source/MovieSceneAnimMixer/Private/MovieSceneAnimationMixerTrack.cpp
@@ -86,32 +86,31 @@ EMovieSceneChannelProxyType UMovieSceneRootMotionSettingsDecoration::PopulateCha
 	EMovieSceneRootMotionTransformMode CurrentTransformMode = GetRootTransformMode();
 
 #if WITH_EDITOR
-	FText Group = LOCTEXT("RootMotionSettings", "Root Motion Settings");
 	int32 SortOrder = 0;
 	{
 		FMovieSceneChannelMetaData ChannelMetaData;
-		ChannelMetaData.SetIdentifiers("Space", LOCTEXT("Space", "Space"), Group);
+		ChannelMetaData.SetIdentifiers("Space", LOCTEXT("Space", "Space"));
 		ChannelMetaData.WeakOwningObject = this;
 		ChannelMetaData.SortOrder = SortOrder++;
 		OutProxyData.Add(RootMotionSpace, ChannelMetaData, TMovieSceneExternalValue<uint8>());
 	}
 	{
 		FMovieSceneChannelMetaData ChannelMetaData;
-		ChannelMetaData.SetIdentifiers("RootMotionSource", LOCTEXT("Source", "Source"), Group);
+		ChannelMetaData.SetIdentifiers("RootMotionSource", LOCTEXT("Source", "Source"));
 		ChannelMetaData.WeakOwningObject = this;
 		ChannelMetaData.SortOrder = SortOrder++;
 		OutProxyData.Add(RootMotionSourceChannel, ChannelMetaData, TMovieSceneExternalValue<uint8>());
 	}
 	{
 		FMovieSceneChannelMetaData ChannelMetaData;
-		ChannelMetaData.SetIdentifiers("Root Mode", LOCTEXT("Mode", "Mode"), Group);
+		ChannelMetaData.SetIdentifiers("Root Mode", LOCTEXT("Mode", "Mode"));
 		ChannelMetaData.WeakOwningObject = this;
 		ChannelMetaData.SortOrder = SortOrder++;
 		OutProxyData.Add(TransformMode, ChannelMetaData, TMovieSceneExternalValue<uint8>());
 	}
 	{
 		FMovieSceneChannelMetaData ChannelMetaData;
-		ChannelMetaData.SetIdentifiers("KeepState", LOCTEXT("KeepState", "Keep State"), Group);
+		ChannelMetaData.SetIdentifiers("KeepState", LOCTEXT("KeepState", "Keep State"));
 		ChannelMetaData.WeakOwningObject = this;
 		ChannelMetaData.SortOrder = SortOrder++;
 		OutProxyData.Add(KeepStateChannel, ChannelMetaData, TMovieSceneExternalValue<uint8>());
@@ -120,10 +119,11 @@ EMovieSceneChannelProxyType UMovieSceneRootMotionSettingsDecoration::PopulateCha
 	if (CurrentTransformMode != EMovieSceneRootMotionTransformMode::Asset)
 	{
 		FName SliderExponent("SliderExponent");
+		FText OffsetGroup = LOCTEXT("RootMotionOffset", "Offset");
 
 		{
 			FMovieSceneChannelMetaData ChannelMetaData;
-			ChannelMetaData.SetIdentifiers("RootBaseLocation.X", AxisDisplayInfo::GetAxisDisplayName(EAxisList::X), Group);
+			ChannelMetaData.SetIdentifiers("RootBaseLocation.X", AxisDisplayInfo::GetAxisDisplayName(EAxisList::X), OffsetGroup);
 			ChannelMetaData.Color = AxisDisplayInfo::GetAxisColor(EAxisList::X);
 			ChannelMetaData.WeakOwningObject = this;
 			ChannelMetaData.SortOrder = SortOrder++;
@@ -132,7 +132,7 @@ EMovieSceneChannelProxyType UMovieSceneRootMotionSettingsDecoration::PopulateCha
 		}
 		{
 			FMovieSceneChannelMetaData ChannelMetaData;
-			ChannelMetaData.SetIdentifiers("RootBaseLocation.Y", AxisDisplayInfo::GetAxisDisplayName(EAxisList::Y), Group);
+			ChannelMetaData.SetIdentifiers("RootBaseLocation.Y", AxisDisplayInfo::GetAxisDisplayName(EAxisList::Y), OffsetGroup);
 			ChannelMetaData.Color = AxisDisplayInfo::GetAxisColor(EAxisList::Y);
 			ChannelMetaData.WeakOwningObject = this;
 			ChannelMetaData.SortOrder = SortOrder++;
@@ -141,7 +141,7 @@ EMovieSceneChannelProxyType UMovieSceneRootMotionSettingsDecoration::PopulateCha
 		}
 		{
 			FMovieSceneChannelMetaData ChannelMetaData;
-			ChannelMetaData.SetIdentifiers("RootBaseLocation.Z", AxisDisplayInfo::GetAxisDisplayName(EAxisList::Z), Group);
+			ChannelMetaData.SetIdentifiers("RootBaseLocation.Z", AxisDisplayInfo::GetAxisDisplayName(EAxisList::Z), OffsetGroup);
 			ChannelMetaData.Color = AxisDisplayInfo::GetAxisColor(EAxisList::Z);
 			ChannelMetaData.WeakOwningObject = this;
 			ChannelMetaData.SortOrder = SortOrder++;
@@ -150,7 +150,7 @@ EMovieSceneChannelProxyType UMovieSceneRootMotionSettingsDecoration::PopulateCha
 		}
 		{
 			FMovieSceneChannelMetaData ChannelMetaData;
-			ChannelMetaData.SetIdentifiers("RootBaseRotation.X", LOCTEXT("RotationX", "Roll"), Group);
+			ChannelMetaData.SetIdentifiers("RootBaseRotation.X", LOCTEXT("RotationX", "Roll"), OffsetGroup);
 			ChannelMetaData.Color = AxisDisplayInfo::GetAxisColor(EAxisList::X);
 			ChannelMetaData.WeakOwningObject = this;
 			ChannelMetaData.SortOrder = SortOrder++;
@@ -159,7 +159,7 @@ EMovieSceneChannelProxyType UMovieSceneRootMotionSettingsDecoration::PopulateCha
 		}
 		{
 			FMovieSceneChannelMetaData ChannelMetaData;
-			ChannelMetaData.SetIdentifiers("RootBaseRotation.Y", LOCTEXT("RotationY", "Pitch"), Group);
+			ChannelMetaData.SetIdentifiers("RootBaseRotation.Y", LOCTEXT("RotationY", "Pitch"), OffsetGroup);
 			ChannelMetaData.Color = AxisDisplayInfo::GetAxisColor(EAxisList::Y);
 			ChannelMetaData.WeakOwningObject = this;
 			ChannelMetaData.SortOrder = SortOrder++;
@@ -168,7 +168,7 @@ EMovieSceneChannelProxyType UMovieSceneRootMotionSettingsDecoration::PopulateCha
 		}
 		{
 			FMovieSceneChannelMetaData ChannelMetaData;
-			ChannelMetaData.SetIdentifiers("RootBaseRotation.Z", LOCTEXT("RotationZ", "Yaw"), Group);
+			ChannelMetaData.SetIdentifiers("RootBaseRotation.Z", LOCTEXT("RotationZ", "Yaw"), OffsetGroup);
 			ChannelMetaData.Color = AxisDisplayInfo::GetAxisColor(EAxisList::Z);
 			ChannelMetaData.WeakOwningObject = this;
 			ChannelMetaData.SortOrder = SortOrder++;
@@ -185,8 +185,7 @@ EMovieSceneChannelProxyType UMovieSceneRootMotionSettingsDecoration::PopulateCha
 		ChannelMetaData.SetIdentifiers("BoneMatch",
 			bMatchValid
 				? LOCTEXT("BoneMatch", "Bone Match")
-				: LOCTEXT("BoneMatchInvalid", "Bone Match (invalid)"),
-			Group);
+				: LOCTEXT("BoneMatchInvalid", "Bone Match (invalid)"));
 		ChannelMetaData.WeakOwningObject = this;
 		ChannelMetaData.SortOrder = SortOrder++;
 		OutProxyData.Add(BoneMatchChannel, ChannelMetaData);
diff --git a/Engine/Source/Editor/Sequencer/Private/MVVM/ViewModels/DecorationOutlinerModel.cpp b/Engine/Source/Editor/Sequencer/Private/MVVM/ViewModels/DecorationOutlinerModel.cpp
--- a/Engine/Source/Editor/Sequencer/Private/MVVM/ViewModels/DecorationOutlinerModel.cpp
+++ b/Engine/Source/Editor/Sequencer/Private/MVVM/ViewModels/DecorationOutlinerModel.cpp
@@ -222,7 +222,12 @@ void FDecorationOutlinerModel::OnConstruct()
 			return A.MetaData->SortOrder < B.MetaData->SortOrder;
 		});
 
-		// Add channels directly under this decoration model
+		// Channels with a non-empty Group are placed under a collapsible
+		// FCategoryGroupModel sub-group. Channels with no Group are added
+		// directly under this decoration model.
+		TMap<FName, TSharedPtr<FCategoryGroupModel>> CategoryGroupModels;
+		TMap<FName, TSharedPtr<FViewModel>> CategoryGroupTails;
+
 		TSharedPtr<FViewModel> ChannelTail;
 		for (const FChannelSortData& ChannelData : SortedChannels)
 		{
@@ -230,13 +235,47 @@ void FDecorationOutlinerModel::OnConstruct()
 			FMovieSceneChannelHandle ChannelHandle = ChannelData.Handle;
 
 			FName ChannelName = MetaData.Name;
+
+			const bool bIsSubGroup = !MetaData.Group.IsEmpty();
+
+			// Find or create the category group model for sub-grouped channels
+			TSharedPtr<FCategoryGroupModel> CategoryGroup;
+			if (bIsSubGroup)
+			{
+				FName GroupName(*MetaData.Group.ToString());
+				TSharedPtr<FCategoryGroupModel>& ExistingCategoryGroup = CategoryGroupModels.FindOrAdd(GroupName);
+				if (!ExistingCategoryGroup)
+				{
+					ExistingCategoryGroup = MakeShared<FCategoryGroupModel>(GroupName, MetaData.Group, MetaData.GetGroupTooltipTextDelegate);
+					OutlinerChildren.InsertChild(ExistingCategoryGroup, ChannelTail);
+					ChannelTail = ExistingCategoryGroup;
+				}
+				CategoryGroup = ExistingCategoryGroup;
+			}
+
+			// Search for an existing channel group model in the appropriate parent
 			TSharedPtr<FChannelGroupOutlinerModel> ChannelGroupModel;
-			for (TSharedPtr<FChannelGroupOutlinerModel> ExistingGroup : OutlinerChildren.IterateSubList<FChannelGroupOutlinerModel>())
+			if (bIsSubGroup)
+			{
+				FViewModelChildren CategoryChildren = CategoryGroup->GetChildList(EViewModelListType::Outliner);
+				for (TSharedPtr<FChannelGroupOutlinerModel> ExistingGroup : CategoryChildren.IterateSubList<FChannelGroupOutlinerModel>())
+				{
+					if (ExistingGroup && ExistingGroup->GetChannelName() == ChannelName)
+					{
+						ChannelGroupModel = ExistingGroup;
+						break;
+					}
+				}
+			}
+			else
 			{
-				if (ExistingGroup && ExistingGroup->GetChannelName() == ChannelName)
+				for (TSharedPtr<FChannelGroupOutlinerModel> ExistingGroup : OutlinerChildren.IterateSubList<FChannelGroupOutlinerModel>())
 				{
-					ChannelGroupModel = ExistingGroup;
-					break;
+					if (ExistingGroup && ExistingGroup->GetChannelName() == ChannelName)
+					{
+						ChannelGroupModel = ExistingGroup;
+						break;
+					}
 				}
 			}
 
@@ -250,8 +289,18 @@ void FDecorationOutlinerModel::OnConstruct()
 				ChannelGroupModel->OnRecycle();
 			}
 
-			OutlinerChildren.InsertChild(ChannelGroupModel, ChannelTail);
-			ChannelTail = ChannelGroupModel;
+			if (bIsSubGroup)
+			{
+				FName GroupName(*MetaData.Group.ToString());
+				TSharedPtr<FViewModel>& GroupTail = CategoryGroupTails.FindOrAdd(GroupName);
+				CategoryGroup->GetChildList(EViewModelListType::Outliner).InsertChild(ChannelGroupModel, GroupTail);
+				GroupTail = ChannelGroupModel;
+			}
+			else
+			{
+				OutlinerChildren.InsertChild(ChannelGroupModel, ChannelTail);
+				ChannelTail = ChannelGroupModel;
+			}
 
 			TSharedPtr<FChannelModel> ChannelModel = MakeShared<FChannelModel>(ChannelName, SectionInterface, ChannelHandle);
 			ChannelGroupModel->AddChannel(ChannelModel);
