Sequencer: Fix keying/interrogation of transforms when using anim mixer and root motion.

Previously this was done by a hardcoded check of the Skeletal Animation System to find the applied root motion so we could subtract it before keying. However now that we could be running the mixer I've made a modular feature.
Now both the anim mixer and the skeletal animation system subscribe as root motion offset providers and the transform track editor looks for whoever has implemented the modular feature and interacts with that. This allows the backwards compatibility behavior while making the mixer work here.

#Synced-CL 51609808

diff --git a/Engine/Plugins/Experimental/MovieSceneAnimMixer/Source/MovieSceneAnimMixer/Private/Systems/MovieSceneRootMotionSystem.cpp b/Engine/Plugins/Experimental/MovieSceneAnimMixer/Source/MovieSceneAnimMixer/Private/Systems/MovieSceneRootMotionSystem.cpp
--- a/Engine/Plugins/Experimental/MovieSceneAnimMixer/Source/MovieSceneAnimMixer/Private/Systems/MovieSceneRootMotionSystem.cpp
+++ b/Engine/Plugins/Experimental/MovieSceneAnimMixer/Source/MovieSceneAnimMixer/Private/Systems/MovieSceneRootMotionSystem.cpp
@@ -528,6 +528,9 @@ void FAnimNextStoreRootTransformTask::Execute(UE::UAF::FEvaluationVM& VM) const
 		}
 	}
 
+	// Store the root motion offset so the transform track editor can subtract it during auto-key
+	RootMotionData->AppliedRootMotionOffset = RootTransform * RootMotionData->ActorTransform.Inverse();
+
 	if (FinalDestination != EMovieSceneRootMotionDestination::Attribute)
 	{
 		TWeakPtr<FMovieSceneMixerRootMotionComponentData> WeakRootMotion = RootMotionData;
diff --git a/Engine/Plugins/Experimental/MovieSceneAnimMixer/Source/MovieSceneAnimMixer/Public/Systems/MovieSceneAnimMixerSystem.h b/Engine/Plugins/Experimental/MovieSceneAnimMixer/Source/MovieSceneAnimMixer/Public/Systems/MovieSceneAnimMixerSystem.h
--- a/Engine/Plugins/Experimental/MovieSceneAnimMixer/Source/MovieSceneAnimMixer/Public/Systems/MovieSceneAnimMixerSystem.h
+++ b/Engine/Plugins/Experimental/MovieSceneAnimMixer/Source/MovieSceneAnimMixer/Public/Systems/MovieSceneAnimMixerSystem.h
@@ -179,7 +179,7 @@ public:
 
 	static TInstancedStruct<FMovieSceneMixedAnimationTarget> ResolveAnimationTarget(FObjectKey ObjectKey, const TInstancedStruct<FMovieSceneMixedAnimationTarget>& InTarget, const FMovieSceneSkeletalAnimationComponentData* SkelAnimData = nullptr);
 
-	TSharedPtr<FMovieSceneMixerRootMotionComponentData> FindRootMotion(FObjectKey InObject) const;
+	MOVIESCENEANIMMIXER_API TSharedPtr<FMovieSceneMixerRootMotionComponentData> FindRootMotion(FObjectKey InObject) const;
 	void AssignRootMotion(FObjectKey InObjectKey, TSharedPtr<FMovieSceneMixerRootMotionComponentData> RootMotion);
 
 	void PreInitializeAllRootMotion();
diff --git a/Engine/Plugins/Experimental/MovieSceneAnimMixer/Source/MovieSceneAnimMixer/Public/Systems/MovieSceneRootMotionSystem.h b/Engine/Plugins/Experimental/MovieSceneAnimMixer/Source/MovieSceneAnimMixer/Public/Systems/MovieSceneRootMotionSystem.h
--- a/Engine/Plugins/Experimental/MovieSceneAnimMixer/Source/MovieSceneAnimMixer/Public/Systems/MovieSceneRootMotionSystem.h
+++ b/Engine/Plugins/Experimental/MovieSceneAnimMixer/Source/MovieSceneAnimMixer/Public/Systems/MovieSceneRootMotionSystem.h
@@ -130,6 +130,10 @@ private:
 
 public:
 
+	// The root motion offset applied to the component, relative to the actor heading.
+	// Stored each evaluation so the transform track editor can subtract it during auto-key.
+	FTransform AppliedRootMotionOffset = FTransform::Identity;
+
 	bool bComponentSpaceRoot = false;
 	bool bActorTransformSet = false;
 
diff --git a/Engine/Plugins/Experimental/MovieSceneAnimMixer/Source/MovieSceneAnimMixerEditor/Private/MovieSceneAnimMixerEditorModule.cpp b/Engine/Plugins/Experimental/MovieSceneAnimMixer/Source/MovieSceneAnimMixerEditor/Private/MovieSceneAnimMixerEditorModule.cpp
--- a/Engine/Plugins/Experimental/MovieSceneAnimMixer/Source/MovieSceneAnimMixerEditor/Private/MovieSceneAnimMixerEditorModule.cpp
+++ b/Engine/Plugins/Experimental/MovieSceneAnimMixer/Source/MovieSceneAnimMixerEditor/Private/MovieSceneAnimMixerEditorModule.cpp
@@ -4,6 +4,7 @@
 
 #include "AnimBlueprintExtension_SequencerMixerTarget.h"
 #include "AnimMixerDecorationMenuProviders.h"
+#include "IMovieSceneRootMotionOffsetProvider.h"
 #include "AnimMixerTargetMenuProviders.h"
 #include "AnimTransitionSectionInterface.h"
 #include "IAnimGraph_SequencerMixerTargetConnector.h"
@@ -21,6 +22,10 @@
 #include "RootMotionSettingsDecorationEditor.h"
 #include "Channels/BuiltInChannelEditors.h"
 #include "MVVM/AnimationMixerTrackModel.h"
+#include "Systems/MovieSceneAnimMixerSystem.h"
+#include "Systems/MovieSceneRootMotionSystem.h"
+#include "EntitySystem/MovieSceneEntitySystemLinker.h"
+#include "GameFramework/Actor.h"
 
 
 namespace Sequencer
@@ -158,6 +163,39 @@ inline bool EvaluateChannel(const FMovieSceneBoneMatchChannel* InChannel, FFrame
 
 namespace UE::MovieScene
 {
+	FTransform FAnimMixerRootMotionOffsetProvider::GetRootMotionOffset(const UMovieSceneEntitySystemLinker* Linker, UObject* AnimatedObject) const
+	{
+		if (!Linker || !AnimatedObject)
+		{
+			return FTransform::Identity;
+		}
+
+		const UMovieSceneAnimMixerSystem* AnimMixerSystem = Linker->FindSystem<UMovieSceneAnimMixerSystem>();
+		if (!AnimMixerSystem)
+		{
+			return FTransform::Identity;
+		}
+
+		TSharedPtr<FMovieSceneMixerRootMotionComponentData> RootMotionData = AnimMixerSystem->FindRootMotion(FObjectKey(AnimatedObject));
+		if (!RootMotionData)
+		{
+			if (const AActor* Actor = Cast<AActor>(AnimatedObject))
+			{
+				if (USceneComponent* RootComponent = Actor->GetRootComponent())
+				{
+					RootMotionData = AnimMixerSystem->FindRootMotion(FObjectKey(RootComponent));
+				}
+			}
+		}
+
+		if (RootMotionData)
+		{
+			return RootMotionData->AppliedRootMotionOffset;
+		}
+
+		return FTransform::Identity;
+	}
+
 	FLazyName SequencerModuleName("Sequencer");
 
 	void FMovieSceneAnimMixerEditorModule::StartupModule()
@@ -222,10 +260,13 @@ namespace UE::MovieScene
 		// Register decoration menu providers
 		MaskDecorationMenuProvider = MakeUnique<FMaskDecorationMenuProvider>();
 		IModularFeatures::Get().RegisterModularFeature(IMovieSceneDecorationMenuProvider::GetModularFeatureName(), MaskDecorationMenuProvider.Get());
+		IModularFeatures::Get().RegisterModularFeature(IMovieSceneRootMotionOffsetProvider::GetModularFeatureName(), &AnimMixerRootMotionProvider);
 	}
 
 	void FMovieSceneAnimMixerEditorModule::ShutdownModule()
 	{
+		IModularFeatures::Get().UnregisterModularFeature(IMovieSceneRootMotionOffsetProvider::GetModularFeatureName(), &AnimMixerRootMotionProvider);
+
 		// Unregister decoration editors
 		if (RootMotionSettingsDecorationEditor)
 		{
diff --git a/Engine/Plugins/Experimental/MovieSceneAnimMixer/Source/MovieSceneAnimMixerEditor/Private/MovieSceneAnimMixerEditorModule.h b/Engine/Plugins/Experimental/MovieSceneAnimMixer/Source/MovieSceneAnimMixerEditor/Private/MovieSceneAnimMixerEditorModule.h
--- a/Engine/Plugins/Experimental/MovieSceneAnimMixer/Source/MovieSceneAnimMixerEditor/Private/MovieSceneAnimMixerEditorModule.h
+++ b/Engine/Plugins/Experimental/MovieSceneAnimMixer/Source/MovieSceneAnimMixerEditor/Private/MovieSceneAnimMixerEditorModule.h
@@ -4,6 +4,7 @@
 
 #include "CoreMinimal.h"
 #include "IAnimGraph_SequencerMixerTargetConnector.h"
+#include "IMovieSceneRootMotionOffsetProvider.h"
 #include "Modules/ModuleInterface.h"
 #include "MovieSceneAnimMixerEditor/MovieSceneAnimationMaskDecorationEditor.h"
 
@@ -21,6 +22,12 @@ namespace UE::Sequencer
 
 namespace UE::MovieScene
 {
+	class FAnimMixerRootMotionOffsetProvider : public IMovieSceneRootMotionOffsetProvider
+	{
+	public:
+		virtual FTransform GetRootMotionOffset(const UMovieSceneEntitySystemLinker* Linker, UObject* AnimatedObject) const override;
+	};
+
 	class FMovieSceneAnimMixerEditorModule : public IModuleInterface, public IAnimGraph_SequencerMixerTargetConnector
 	{
 	public:
@@ -48,6 +55,7 @@ namespace UE::MovieScene
 
 		// Decoration menu providers
 		TUniquePtr<FMaskDecorationMenuProvider> MaskDecorationMenuProvider;
+		FAnimMixerRootMotionOffsetProvider AnimMixerRootMotionProvider;
 	};
 }
 
diff --git a/Engine/Source/Editor/MovieSceneTools/Private/MovieSceneToolsModule.cpp b/Engine/Source/Editor/MovieSceneTools/Private/MovieSceneToolsModule.cpp
--- a/Engine/Source/Editor/MovieSceneTools/Private/MovieSceneToolsModule.cpp
+++ b/Engine/Source/Editor/MovieSceneTools/Private/MovieSceneToolsModule.cpp
@@ -145,9 +145,37 @@
 #include "Conditions/MovieSceneDirectorBlueprintConditionUtils.h"
 #include "Conditions/MovieSceneDirectorBlueprintConditionCustomization.h"
 #include "Conditions/MovieScenePlatformConditionCustomization.h"
+#include "IMovieSceneRootMotionOffsetProvider.h"
+#include "Systems/MovieSceneSkeletalAnimationSystem.h"
+#include "EntitySystem/MovieSceneEntitySystemLinker.h"
 
 #define LOCTEXT_NAMESPACE "FMovieSceneToolsModule"
 
+FTransform FSkeletalAnimationRootMotionOffsetProvider::GetRootMotionOffset(const UMovieSceneEntitySystemLinker* Linker, UObject* AnimatedObject) const
+{
+	if (!Linker)
+	{
+		return FTransform::Identity;
+	}
+	if (const UMovieSceneSkeletalAnimationSystem* System = Linker->FindSystem<UMovieSceneSkeletalAnimationSystem>())
+	{
+		return System->GetRootMotionOffset(AnimatedObject);
+	}
+	return FTransform::Identity;
+}
+
+void FSkeletalAnimationRootMotionOffsetProvider::UpdateRootMotionOffset(UMovieSceneEntitySystemLinker* Linker, UObject* AnimatedObject)
+{
+	if (!Linker)
+	{
+		return;
+	}
+	if (UMovieSceneSkeletalAnimationSystem* System = Linker->FindSystem<UMovieSceneSkeletalAnimationSystem>())
+	{
+		System->UpdateRootMotionOffset(AnimatedObject);
+	}
+}
+
 TAutoConsoleVariable<bool> CVarDuplicateLinkedAnimSequence(TEXT("Sequencer.DuplicateLinkedAnimSequence"), false, TEXT("When true when we duplicate a level sequence that has a linked anim sequence it will duplicate and link the anim sequencel, if false we leave any link alone."));
 
 
@@ -330,10 +358,14 @@ void FMovieSceneToolsModule::StartupModule()
 	// register UTransformableComponentHandle animatable interface
 	FConstraintChannelInterfaceRegistry& ConstraintChannelInterfaceRegistry = FConstraintChannelInterfaceRegistry::Get();
 	ConstraintChannelInterfaceRegistry.RegisterConstraintChannelInterface<UTransformableComponentHandle>(MakeUnique<FComponentConstraintChannelInterface>());
+
+	// Register legacy skeletal animation root motion offset provider
+	IModularFeatures::Get().RegisterModularFeature(IMovieSceneRootMotionOffsetProvider::GetModularFeatureName(), &SkeletalAnimRootMotionProvider);
 }
 
 void FMovieSceneToolsModule::ShutdownModule()
 {
+	IModularFeatures::Get().UnregisterModularFeature(IMovieSceneRootMotionOffsetProvider::GetModularFeatureName(), &SkeletalAnimRootMotionProvider);
 	UMovieScene::FixupDynamicBindingPayloadParameterNameEvent.Remove(FixupDynamicBindingPayloadParameterNameHandle);
 	UMovieSceneEventSectionBase::FixupPayloadParameterNameEvent.Remove(FixupEventSectionPayloadParameterNameHandle);
 	UMovieSceneEventSectionBase::UpgradeLegacyEventEndpoint = UMovieSceneEventSectionBase::FUpgradeLegacyEventEndpoint();
diff --git a/Engine/Source/Editor/MovieSceneTools/Private/TrackEditors/TransformTrackEditor.cpp b/Engine/Source/Editor/MovieSceneTools/Private/TrackEditors/TransformTrackEditor.cpp
--- a/Engine/Source/Editor/MovieSceneTools/Private/TrackEditors/TransformTrackEditor.cpp
+++ b/Engine/Source/Editor/MovieSceneTools/Private/TrackEditors/TransformTrackEditor.cpp
@@ -54,7 +54,7 @@
 #include "ISequencerObjectChangeListener.h"
 #include "ISequencerPropertyKeyedStatus.h"
 #include "Systems/MovieSceneTransformOriginSystem.h"
-#include "Systems/MovieSceneSkeletalAnimationSystem.h"
+#include "IMovieSceneRootMotionOffsetProvider.h"
 
 #define LOCTEXT_NAMESPACE "MovieScene_TransformTrack"
 
@@ -991,35 +991,45 @@ FTransform F3DTransformTrackEditor::GetTransformOrigin() const
 
 FTransform F3DTransformTrackEditor::GetRootMotionOffset(UObject* AnimatedObject) const
 {
-	FTransform RootMotionOffset;
-
 	const FMovieSceneRootEvaluationTemplateInstance& EvaluationTemplate = GetSequencer()->GetEvaluationTemplate();
 	const UMovieSceneEntitySystemLinker* EntityLinker = EvaluationTemplate.GetEntitySystemLinker();
 	if (!EntityLinker)
 	{
-		return RootMotionOffset;
+		return FTransform::Identity;
 	}
 
-	if (const UMovieSceneSkeletalAnimationSystem* SkeletalAnimationSystem = EntityLinker->FindSystem<UMovieSceneSkeletalAnimationSystem>())
+	// Each provider returns Identity for objects it doesn't manage, so the
+	// first non-identity result is the one that owns this object.
+	TArray<IMovieSceneRootMotionOffsetProvider*> Providers =
+		IModularFeatures::Get().GetModularFeatureImplementations<IMovieSceneRootMotionOffsetProvider>(
+			IMovieSceneRootMotionOffsetProvider::GetModularFeatureName());
+	for (IMovieSceneRootMotionOffsetProvider* Provider : Providers)
 	{
-		return SkeletalAnimationSystem->GetRootMotionOffset(AnimatedObject);
+		FTransform Offset = Provider->GetRootMotionOffset(EntityLinker, AnimatedObject);
+		if (!Offset.Equals(FTransform::Identity))
+		{
+			return Offset;
+		}
 	}
 
-	return RootMotionOffset;
+	return FTransform::Identity;
 }
 
 void F3DTransformTrackEditor::UpdateRootMotionOffset(UObject* AnimatedObject)
 {
 	const FMovieSceneRootEvaluationTemplateInstance& EvaluationTemplate = GetSequencer()->GetEvaluationTemplate();
-	const UMovieSceneEntitySystemLinker* EntityLinker = EvaluationTemplate.GetEntitySystemLinker();
+	UMovieSceneEntitySystemLinker* EntityLinker = EvaluationTemplate.GetEntitySystemLinker();
 	if (!EntityLinker)
 	{
 		return;
 	}
 
-	if (UMovieSceneSkeletalAnimationSystem* SkeletalAnimationSystem = EntityLinker->FindSystem<UMovieSceneSkeletalAnimationSystem>())
+	TArray<IMovieSceneRootMotionOffsetProvider*> Providers =
+		IModularFeatures::Get().GetModularFeatureImplementations<IMovieSceneRootMotionOffsetProvider>(
+			IMovieSceneRootMotionOffsetProvider::GetModularFeatureName());
+	for (IMovieSceneRootMotionOffsetProvider* Provider : Providers)
 	{
-		SkeletalAnimationSystem->UpdateRootMotionOffset(AnimatedObject);
+		Provider->UpdateRootMotionOffset(EntityLinker, AnimatedObject);
 	}
 }
 
diff --git a/Engine/Source/Editor/MovieSceneTools/Public/IMovieSceneRootMotionOffsetProvider.h b/Engine/Source/Editor/MovieSceneTools/Public/IMovieSceneRootMotionOffsetProvider.h
new file mode 100644
--- /dev/null
+++ b/Engine/Source/Editor/MovieSceneTools/Public/IMovieSceneRootMotionOffsetProvider.h
@@ -0,0 +1,31 @@
+// Copyright Epic Games, Inc. All Rights Reserved.
+
+#pragma once
+
+#include "CoreMinimal.h"
+#include "Features/IModularFeatures.h"
+
+class UMovieSceneEntitySystemLinker;
+
+// Modular feature for systems that apply root motion to scene components.
+// Allows the transform track editor to subtract root motion when auto-keying,
+// preventing double-application of root motion offsets.
+class MOVIESCENETOOLS_API IMovieSceneRootMotionOffsetProvider : public IModularFeature
+{
+public:
+	virtual ~IMovieSceneRootMotionOffsetProvider() = default;
+
+	// Returns the root motion offset currently applied to the given object,
+	// or Identity if this provider doesn't handle the object.
+	virtual FTransform GetRootMotionOffset(const UMovieSceneEntitySystemLinker* Linker, UObject* AnimatedObject) const = 0;
+
+	// Recalculates the initial actor transform baseline after the user manually
+	// changes the transform.
+	virtual void UpdateRootMotionOffset(UMovieSceneEntitySystemLinker* Linker, UObject* AnimatedObject) {}
+
+	static FName GetModularFeatureName()
+	{
+		static FName Name(TEXT("MovieSceneRootMotionOffsetProvider"));
+		return Name;
+	}
+};
diff --git a/Engine/Source/Editor/MovieSceneTools/Public/MovieSceneToolsModule.h b/Engine/Source/Editor/MovieSceneTools/Public/MovieSceneToolsModule.h
--- a/Engine/Source/Editor/MovieSceneTools/Public/MovieSceneToolsModule.h
+++ b/Engine/Source/Editor/MovieSceneTools/Public/MovieSceneToolsModule.h
@@ -9,6 +9,7 @@
 #include "Curves/RichCurve.h"
 #include "Containers/UnrealString.h"
 #include "IMovieSceneTools.h"
+#include "IMovieSceneRootMotionOffsetProvider.h"
 
 #define UE_API MOVIESCENETOOLS_API
 
@@ -62,6 +63,15 @@ public:
 	virtual void RegisterKeyStructInstancedPropertyTypeCustomization(TSharedRef<IStructureDetailsView> StructureDetailsView, TWeakObjectPtr<UMovieSceneSection> WeakOwningSection) {};
 };
 
+// Wraps UMovieSceneSkeletalAnimationSystem as an IMovieSceneRootMotionOffsetProvider
+// so the transform track editor can subtract root motion during auto-key.
+class FSkeletalAnimationRootMotionOffsetProvider : public IMovieSceneRootMotionOffsetProvider
+{
+public:
+	UE_API virtual FTransform GetRootMotionOffset(const UMovieSceneEntitySystemLinker* Linker, UObject* AnimatedObject) const override;
+	UE_API virtual void UpdateRootMotionOffset(UMovieSceneEntitySystemLinker* Linker, UObject* AnimatedObject) override;
+};
+
 /**
 * Implements the MovieSceneTools module.
 */
@@ -174,6 +184,8 @@ private:
 
 	TArray<IMovieSceneToolsAnimationBakeHelper*> BakeHelpers;
 	TArray<IMovieSceneToolsKeyStructInstancedPropertyTypeCustomizer*> KeyStructInstancedPropertyTypeCustomizers;
+
+	FSkeletalAnimationRootMotionOffsetProvider SkeletalAnimRootMotionProvider;
 };
 
 #undef UE_API
