Sequencer: Transition resize and remove now respect lock state of source sections

- When resizing the left edge of a transition, block if the TO section is locked
- When resizing the right edge, block if the FROM section is locked
- Moving a transition is blocked if either of its source sections are locked
- Prevents transitions from getting out of sync with locked section bounds

#Synced-CL 51964473

diff --git a/Engine/Plugins/Experimental/MovieSceneAnimMixer/Source/MovieSceneAnimMixer/Private/MovieSceneAnimationMixerTrack.cpp b/Engine/Plugins/Experimental/MovieSceneAnimMixer/Source/MovieSceneAnimMixer/Private/MovieSceneAnimationMixerTrack.cpp
--- a/Engine/Plugins/Experimental/MovieSceneAnimMixer/Source/MovieSceneAnimMixer/Private/MovieSceneAnimationMixerTrack.cpp
+++ b/Engine/Plugins/Experimental/MovieSceneAnimMixer/Source/MovieSceneAnimMixer/Private/MovieSceneAnimationMixerTrack.cpp
@@ -770,6 +770,13 @@ EMovieSceneSectionMovedResult UMovieSceneAnimationMixerTrack::HandleTransitionSe
 		return EMovieSceneSectionMovedResult::None;
 	}
 
+	// If either source section is locked, block the move and restore bounds
+	if (TransitionSection.FromSection->IsLocked() || TransitionSection.ToSection->IsLocked())
+	{
+		TransitionSection.UpdateBoundsFromSourceSections();
+		return EMovieSceneSectionMovedResult::None;
+	}
+
 	// 1. Handle row changes - transition must stay on same row as from/to sections
 	// This prevents vertical dragging of transitions
 	int32 CorrectRowIndex = TransitionSection.FromSection->GetRowIndex();
diff --git a/Engine/Plugins/Experimental/MovieSceneAnimMixer/Source/MovieSceneAnimMixerEditor/Private/AnimTransitionSectionInterface.cpp b/Engine/Plugins/Experimental/MovieSceneAnimMixer/Source/MovieSceneAnimMixerEditor/Private/AnimTransitionSectionInterface.cpp
--- a/Engine/Plugins/Experimental/MovieSceneAnimMixer/Source/MovieSceneAnimMixerEditor/Private/AnimTransitionSectionInterface.cpp
+++ b/Engine/Plugins/Experimental/MovieSceneAnimMixer/Source/MovieSceneAnimMixerEditor/Private/AnimTransitionSectionInterface.cpp
@@ -365,7 +365,13 @@ void FAnimTransitionSectionInterface::ResizeSection(ESequencerSectionResizeMode
 
 	if (ResizeMode == SSRM_LeadingEdge)
 	{
-		// Resizing the left edge - this controls where the TO section starts and transition starts
+		// Resizing the left edge - this controls where the TO section starts and transition starts.
+		// If the TO section is locked, prevent resizing this edge.
+		if (TransitionSection->ToSection->IsLocked())
+		{
+			return;
+		}
+
 		// Cannot go earlier than FROM section's original start
 		FFrameNumber MinFrame = InitialFromRange.HasLowerBound() ? InitialFromRange.GetLowerBoundValue() : TNumericLimits<int32>::Min();
 		FFrameNumber MaxFrame = TransitionSection->HasEndFrame() ? TransitionSection->GetExclusiveEndFrame() - 1 : TNumericLimits<int32>::Max();
@@ -382,7 +388,13 @@ void FAnimTransitionSectionInterface::ResizeSection(ESequencerSectionResizeMode
 	}
 	else // SSRM_TrailingEdge
 	{
-		// Resizing the right edge - this controls where the FROM section ends and transition ends
+		// Resizing the right edge - this controls where the FROM section ends and transition ends.
+		// If the FROM section is locked, prevent resizing this edge.
+		if (TransitionSection->FromSection->IsLocked())
+		{
+			return;
+		}
+
 		// Cannot go later than TO section's original end
 		FFrameNumber MinFrame = TransitionSection->HasStartFrame() ? TransitionSection->GetInclusiveStartFrame() + 1 : TNumericLimits<int32>::Min();
 		FFrameNumber MaxFrame = InitialToRange.HasUpperBound() ? InitialToRange.GetUpperBoundValue() : TNumericLimits<int32>::Max();
