Sequencer: Fixed anim mixer mute/solo/lock/disabled state not being remapped when reordering mix layers

#Synced-CL 51964473

diff --git a/Engine/Plugins/Experimental/MovieSceneAnimMixer/Source/MovieSceneAnimMixerEditor/Private/MVVM/AnimationMixerLayerModel.cpp b/Engine/Plugins/Experimental/MovieSceneAnimMixer/Source/MovieSceneAnimMixerEditor/Private/MVVM/AnimationMixerLayerModel.cpp
--- a/Engine/Plugins/Experimental/MovieSceneAnimMixer/Source/MovieSceneAnimMixerEditor/Private/MVVM/AnimationMixerLayerModel.cpp
+++ b/Engine/Plugins/Experimental/MovieSceneAnimMixer/Source/MovieSceneAnimMixerEditor/Private/MVVM/AnimationMixerLayerModel.cpp
@@ -397,10 +397,30 @@ bool FAnimationMixerLayerModel::SetRowIndex(int32 InRowIndex)
 
 	MixerTrack->Modify();
 
+	// Build a mapping of old row indices before reordering so we can remap
+	// per-row metadata (mute/solo state, track row metadata, etc.)
+	TArray<int32> OldIndices;
+	OldIndices.SetNum(Layers.Num());
+	for (int32 Index = 0; Index < Layers.Num(); ++Index)
+	{
+		OldIndices[Index] = Index;
+	}
+
 	// Move the layer in the array
 	Layers.RemoveAt(CurrentIndex);
 	Layers.Insert(Layer, InRowIndex);
 
+	// The same reorder applied to the old indices gives us the old-index-at-new-position
+	OldIndices.RemoveAt(CurrentIndex);
+	OldIndices.Insert(CurrentIndex, InRowIndex);
+
+	// Build NewToOld map: NewRowIndex -> OldRowIndex
+	TMap<int32, int32> NewToOldRowIndices;
+	for (int32 NewIndex = 0; NewIndex < OldIndices.Num(); ++NewIndex)
+	{
+		NewToOldRowIndices.Add(NewIndex, OldIndices[NewIndex]);
+	}
+
 	// Update row indices for all layers and their contents
 	for (int32 LayerIndex = 0; LayerIndex < Layers.Num(); ++LayerIndex)
 	{
@@ -431,6 +451,9 @@ bool FAnimationMixerLayerModel::SetRowIndex(int32 InRowIndex)
 		}
 	}
 
+	// Remap per-row metadata (mute/solo state, track row decoration, etc.)
+	MixerTrack->OnRowIndicesChanged(NewToOldRowIndices);
+
 	return true;
 }
 
diff --git a/Engine/Source/Runtime/MovieScene/Private/MovieSceneTrack.cpp b/Engine/Source/Runtime/MovieScene/Private/MovieSceneTrack.cpp
--- a/Engine/Source/Runtime/MovieScene/Private/MovieSceneTrack.cpp
+++ b/Engine/Source/Runtime/MovieScene/Private/MovieSceneTrack.cpp
@@ -424,6 +424,32 @@ void UMovieSceneTrack::OnRowIndicesChanged(const TMap<int32, int32>& NewToOldRow
 	}
 	TrackRowMetadata = NewTrackRowMetadata;
 
+	// Remap persistent row-disabled state
+	{
+		TArray<int32> OldRowsDisabled = MoveTemp(RowsDisabled);
+		for (const TPair<int32, int32>& Pair : NewToOldRowIndices)
+		{
+			if (OldRowsDisabled.Contains(Pair.Value))
+			{
+				RowsDisabled.AddUnique(Pair.Key);
+			}
+		}
+	}
+
+#if WITH_EDITORONLY_DATA
+	// Remap local (transient) row-disabled state
+	{
+		TArray<int32> OldLocalRowsDisabled = MoveTemp(LocalRowsDisabled);
+		for (const TPair<int32, int32>& Pair : NewToOldRowIndices)
+		{
+			if (OldLocalRowsDisabled.Contains(Pair.Value))
+			{
+				LocalRowsDisabled.AddUnique(Pair.Key);
+			}
+		}
+	}
+#endif
+
 	if (UMovieSceneTrackRowDecoration* TrackRowDecoration = FindDecoration<UMovieSceneTrackRowDecoration>())
 	{
 		TrackRowDecoration->OnRowIndicesChanged(NewToOldRowIndices);
