Sequencer: Anim Mixer- Fix keyed offset being ignored in AccumulatedOffset root motion mode

* In the runtime conversion task, the AccumulatedOffset and RootTransformOffset code paths were in mutually exclusive if/else branches, so user-keyed offsets had no effect when accumulated offset was active. Apply the offset to the animation-space delta before rebasing onto the accumulated position.
* The accumulated offset cache rebuild baked with bSkipRootMotionConversion (to stay in animation space), which meant the keyed offset wasn't included in the cached delta. Subsequent sections reading from the cache would start at the wrong position. Evaluate the offset channels directly and apply them to the delta during cache rebuild.

#Synced-CL 51964473

diff --git a/Engine/Plugins/Experimental/MovieSceneAnimMixer/Source/MovieSceneAnimMixer/Private/MovieSceneRootMotionSection.cpp b/Engine/Plugins/Experimental/MovieSceneAnimMixer/Source/MovieSceneAnimMixer/Private/MovieSceneRootMotionSection.cpp
--- a/Engine/Plugins/Experimental/MovieSceneAnimMixer/Source/MovieSceneAnimMixer/Private/MovieSceneRootMotionSection.cpp
+++ b/Engine/Plugins/Experimental/MovieSceneAnimMixer/Source/MovieSceneAnimMixer/Private/MovieSceneRootMotionSection.cpp
@@ -267,6 +267,30 @@ void UMovieSceneRootMotionSection::RebuildAccumulatedOffsetCache(
 
 			FTransform Delta = Settings->InitialRootTransform.Inverse() * Result.RootMotionTransform;
 
+			// Apply the user-keyed offset to the delta so that subsequent
+			// sections accumulate on top of the offset. Mirrors the
+			// RootTransformOffset logic in the runtime conversion task.
+			{
+				FVector OffsetLocation = FVector::ZeroVector;
+				FRotator OffsetRotation = FRotator::ZeroRotator;
+				double Val = 0.0;
+				if (Settings->Location[0].Evaluate(Pending.EvalTime, Val)) { OffsetLocation.X = Val; }
+				if (Settings->Location[1].Evaluate(Pending.EvalTime, Val)) { OffsetLocation.Y = Val; }
+				if (Settings->Location[2].Evaluate(Pending.EvalTime, Val)) { OffsetLocation.Z = Val; }
+				if (Settings->Rotation[0].Evaluate(Pending.EvalTime, Val)) { OffsetRotation.Pitch = Val; }
+				if (Settings->Rotation[1].Evaluate(Pending.EvalTime, Val)) { OffsetRotation.Yaw = Val; }
+				if (Settings->Rotation[2].Evaluate(Pending.EvalTime, Val)) { OffsetRotation.Roll = Val; }
+
+				if (!OffsetLocation.IsZero() || !OffsetRotation.IsZero())
+				{
+					FTransform OffsetTransform(OffsetRotation, OffsetLocation);
+					FVector Pivot = Settings->RootOriginLocation;
+					Delta *= FTransform(-Pivot);
+					Delta = Delta * OffsetTransform;
+					Delta *= FTransform(Pivot);
+				}
+			}
+
 			FRootMotionAccumulatedOffset CacheEntry;
 			CacheEntry.StartFrame = Pending.EntryStartFrame;
 			CacheEntry.Weight = 1.0f;
diff --git a/Engine/Plugins/Experimental/MovieSceneAnimMixer/Source/MovieSceneAnimMixer/Private/Systems/MovieSceneRootMotionSystem.cpp b/Engine/Plugins/Experimental/MovieSceneAnimMixer/Source/MovieSceneAnimMixer/Private/Systems/MovieSceneRootMotionSystem.cpp
--- a/Engine/Plugins/Experimental/MovieSceneAnimMixer/Source/MovieSceneAnimMixer/Private/Systems/MovieSceneRootMotionSystem.cpp
+++ b/Engine/Plugins/Experimental/MovieSceneAnimMixer/Source/MovieSceneAnimMixer/Private/Systems/MovieSceneRootMotionSystem.cpp
@@ -329,6 +329,18 @@ void FAnimNextConvertRootMotionToWorldSpaceTask::Execute(UE::UAF::FEvaluationVM&
 	if (EnumHasAnyFlags(Conversions, ESpaceConversions::AccumulatedOffset))
 	{
 		FTransform Delta = InitialAnimRootTransform.Inverse() * RelativeTransform;
+
+		// Apply the user-keyed offset to the animation-space delta before
+		// rebasing onto the accumulated position so that keyed offsets are
+		// additive on top of the accumulated root motion.
+		if (EnumHasAnyFlags(Conversions, ESpaceConversions::RootTransformOffset))
+		{
+			FVector Pivot = RootOffsetOrigin;
+			Delta *= FTransform(-Pivot);
+			Delta = Delta * RootTransform;
+			Delta *= FTransform(Pivot);
+		}
+
 		RelativeTransform = Delta * AccumulatedTransform;
 	}
 	else
