Sequencer: Anim Mixer- Fix zero-width transition sections from sub-frame section overlap

Section endpoints may not align to display frame boundaries (e.g.
skeletal animation sections whose tick length does not evenly divide
the tick resolution). When such sections are snapped to adjacent frame
boundaries they can produce a sub-frame tick-level overlap that
incorrectly creates a zero-width transition section.

Added IsOverlapLargeEnoughForTransition() which requires the overlap
to span at least one full display frame.

#Synced-CL 51964473

diff --git a/Engine/Plugins/Experimental/MovieSceneAnimMixer/Source/MovieSceneAnimMixer/Private/MovieSceneAnimTransitionSectionBase.cpp b/Engine/Plugins/Experimental/MovieSceneAnimMixer/Source/MovieSceneAnimMixer/Private/MovieSceneAnimTransitionSectionBase.cpp
--- a/Engine/Plugins/Experimental/MovieSceneAnimMixer/Source/MovieSceneAnimMixer/Private/MovieSceneAnimTransitionSectionBase.cpp
+++ b/Engine/Plugins/Experimental/MovieSceneAnimMixer/Source/MovieSceneAnimMixer/Private/MovieSceneAnimTransitionSectionBase.cpp
@@ -49,14 +49,11 @@ bool UMovieSceneAnimTransitionSectionBase::IsValid() const
 		return false;
 	}
 
-	// There must be an overlap between the sections
+	// There must be an overlap between the sections that spans at least one display frame.
+	// Sub-frame overlaps can occur when section endpoints don't align to display frame
+	// boundaries and should not be treated as valid transitions.
 	TRange<FFrameNumber> Overlap = TRange<FFrameNumber>::Intersection(FromRange, ToRange);
-	if (Overlap.IsEmpty())
-	{
-		return false;
-	}
-
-	return true;
+	return UMovieSceneAnimationMixerTrack::IsOverlapLargeEnoughForTransition(Overlap, this);
 }
 
 TRange<FFrameNumber> UMovieSceneAnimTransitionSectionBase::ComputeOverlapRange() const
diff --git a/Engine/Plugins/Experimental/MovieSceneAnimMixer/Source/MovieSceneAnimMixer/Private/MovieSceneAnimationMixerTrack.cpp b/Engine/Plugins/Experimental/MovieSceneAnimMixer/Source/MovieSceneAnimMixer/Private/MovieSceneAnimationMixerTrack.cpp
--- a/Engine/Plugins/Experimental/MovieSceneAnimMixer/Source/MovieSceneAnimMixer/Private/MovieSceneAnimationMixerTrack.cpp
+++ b/Engine/Plugins/Experimental/MovieSceneAnimMixer/Source/MovieSceneAnimMixer/Private/MovieSceneAnimationMixerTrack.cpp
@@ -493,6 +493,40 @@ TArray<UMovieSceneAnimTransitionSectionBase*> UMovieSceneAnimationMixerTrack::Fi
 	return Result;
 }
 
+bool UMovieSceneAnimationMixerTrack::IsOverlapLargeEnoughForTransition(const TRange<FFrameNumber>& Overlap, const UObject* ContextObject)
+{
+	if (Overlap.IsEmpty() || !Overlap.HasLowerBound() || !Overlap.HasUpperBound())
+	{
+		return false;
+	}
+
+	int32 OverlapTicks = Overlap.GetUpperBoundValue().Value - Overlap.GetLowerBoundValue().Value;
+	if (OverlapTicks <= 0)
+	{
+		return false;
+	}
+
+	// Section endpoints may not align to display frame boundaries (e.g. skeletal animation
+	// sections whose tick length doesn't evenly divide the tick resolution). When such
+	// sections are snapped to adjacent frame boundaries they can end up with a sub-frame
+	// overlap that should not produce a transition. Require the overlap to span at least
+	// one full display frame.
+	if (const UMovieScene* MovieScene = ContextObject ? ContextObject->GetTypedOuter<UMovieScene>() : nullptr)
+	{
+		const double DisplayRateFPS = MovieScene->GetDisplayRate().AsDecimal();
+		if (DisplayRateFPS > 0.0)
+		{
+			int32 MinOverlapTicks = FMath::FloorToInt32(MovieScene->GetTickResolution().AsDecimal() / DisplayRateFPS);
+			if (OverlapTicks < MinOverlapTicks)
+			{
+				return false;
+			}
+		}
+	}
+
+	return true;
+}
+
 TArray<UMovieSceneSection*> UMovieSceneAnimationMixerTrack::FindOverlappingSections(UMovieSceneSection* Section) const
 {
 	TArray<UMovieSceneSection*> Result;
@@ -516,9 +550,10 @@ TArray<UMovieSceneSection*> UMovieSceneAnimationMixerTrack::FindOverlappingSecti
 				continue;
 			}
 
-			// Check for overlap
+			// Check for overlap (must be at least one display frame to avoid
+			// zero-width transitions from sub-frame endpoint misalignment)
 			TRange<FFrameNumber> Overlap = TRange<FFrameNumber>::Intersection(SectionRange, OtherSection->GetRange());
-			if (!Overlap.IsEmpty())
+			if (IsOverlapLargeEnoughForTransition(Overlap, this))
 			{
 				Result.Add(OtherSection);
 			}
@@ -603,7 +638,7 @@ UMovieSceneAnimTransitionSectionBase* UMovieSceneAnimationMixerTrack::CreateTran
 		ToSection->GetRange()
 	);
 
-	if (Overlap.IsEmpty())
+	if (!IsOverlapLargeEnoughForTransition(Overlap, this))
 	{
 		return nullptr;
 	}
diff --git a/Engine/Plugins/Experimental/MovieSceneAnimMixer/Source/MovieSceneAnimMixer/Public/MovieSceneAnimationMixerTrack.h b/Engine/Plugins/Experimental/MovieSceneAnimMixer/Source/MovieSceneAnimMixer/Public/MovieSceneAnimationMixerTrack.h
--- a/Engine/Plugins/Experimental/MovieSceneAnimMixer/Source/MovieSceneAnimMixer/Public/MovieSceneAnimationMixerTrack.h
+++ b/Engine/Plugins/Experimental/MovieSceneAnimMixer/Source/MovieSceneAnimMixer/Public/MovieSceneAnimationMixerTrack.h
@@ -287,6 +287,13 @@ public:
 	/** Find sections that overlap with the given section on the same layer (excluding transitions) */
 	MOVIESCENEANIMMIXER_API TArray<UMovieSceneSection*> FindOverlappingSections(UMovieSceneSection* Section) const;
 
+	/**
+	 * Returns true if the given overlap range is large enough to warrant a transition.
+	 * Sub-frame overlaps (from section endpoints that don't align to display frame
+	 * boundaries) are rejected so that abutting sections don't produce zero-width transitions.
+	 */
+	MOVIESCENEANIMMIXER_API static bool IsOverlapLargeEnoughForTransition(const TRange<FFrameNumber>& Overlap, const UObject* ContextObject);
+
 	/** Count how many non-transition sections overlap the given range on the specified row */
 	MOVIESCENEANIMMIXER_API int32 CountSectionsOverlappingRange(const TRange<FFrameNumber>& Range, int32 RowIndex) const;
 
