AnimMixer: support for new gizmos in UAnimationMixerTrackEditMode

#rnx

#Synced-CL 52129201

diff --git a/Engine/Plugins/Experimental/MovieSceneAnimMixer/Source/MovieSceneAnimMixerEditor/MovieSceneAnimMixerEditor.Build.cs b/Engine/Plugins/Experimental/MovieSceneAnimMixer/Source/MovieSceneAnimMixerEditor/MovieSceneAnimMixerEditor.Build.cs
--- a/Engine/Plugins/Experimental/MovieSceneAnimMixer/Source/MovieSceneAnimMixerEditor/MovieSceneAnimMixerEditor.Build.cs
+++ b/Engine/Plugins/Experimental/MovieSceneAnimMixer/Source/MovieSceneAnimMixerEditor/MovieSceneAnimMixerEditor.Build.cs
@@ -45,6 +45,7 @@ namespace UnrealBuildTool.Rules
 					"UAFEditor",
 					"UAFUncookedOnly",
 					"UnrealEd",
+					"EditorInteractiveToolsFramework"
 				}
 			);
 		}
diff --git a/Engine/Plugins/Experimental/MovieSceneAnimMixer/Source/MovieSceneAnimMixerEditor/Private/AnimationMixerTrackEditMode.cpp b/Engine/Plugins/Experimental/MovieSceneAnimMixer/Source/MovieSceneAnimMixerEditor/Private/AnimationMixerTrackEditMode.cpp
--- a/Engine/Plugins/Experimental/MovieSceneAnimMixer/Source/MovieSceneAnimMixerEditor/Private/AnimationMixerTrackEditMode.cpp
+++ b/Engine/Plugins/Experimental/MovieSceneAnimMixer/Source/MovieSceneAnimMixerEditor/Private/AnimationMixerTrackEditMode.cpp
@@ -536,6 +536,16 @@ bool UAnimationMixerTrackEditMode::EndTracking(FEditorViewportClient* InViewport
 	return EndDragInternal();
 }
 
+bool UAnimationMixerTrackEditMode::BeginTransform(const FGizmoState& InState)
+{
+	return BeginDragInternal();
+}
+
+bool UAnimationMixerTrackEditMode::EndTransform(const FGizmoState& InState)
+{
+	return EndDragInternal();
+}
+
 bool UAnimationMixerTrackEditMode::InputDelta(FEditorViewportClient* InViewportClient, FViewport* InViewport, FVector& InDrag, FRotator& InRot, FVector& InScale)
 {
 	const bool bCtrlDown = InViewport->KeyState(EKeys::LeftControl) || InViewport->KeyState(EKeys::RightControl);
diff --git a/Engine/Plugins/Experimental/MovieSceneAnimMixer/Source/MovieSceneAnimMixerEditor/Private/AnimationMixerTrackEditMode.h b/Engine/Plugins/Experimental/MovieSceneAnimMixer/Source/MovieSceneAnimMixerEditor/Private/AnimationMixerTrackEditMode.h
--- a/Engine/Plugins/Experimental/MovieSceneAnimMixer/Source/MovieSceneAnimMixerEditor/Private/AnimationMixerTrackEditMode.h
+++ b/Engine/Plugins/Experimental/MovieSceneAnimMixer/Source/MovieSceneAnimMixerEditor/Private/AnimationMixerTrackEditMode.h
@@ -3,6 +3,7 @@
 #pragma once
 
 #include "CoreMinimal.h"
+#include "GizmoEdModeInterface.h"
 #include "Components/SkeletalMeshComponent.h"
 #include "Tools/LegacyEdModeWidgetHelpers.h"
 #include "HitProxies.h"
@@ -80,7 +81,7 @@ struct FMixerSelectedRootData
 // Edit mode for animation mixer tracks providing trajectory drawing,
 // per-section skeleton visualization, and gizmo-based root motion offset keying.
 UCLASS()
-class UAnimationMixerTrackEditMode : public UBaseLegacyWidgetEdMode
+class UAnimationMixerTrackEditMode : public UBaseLegacyWidgetEdMode, public IGizmoEdModeInterface
 {
 	GENERATED_BODY()
 public:
@@ -101,7 +102,9 @@ public:
 	virtual bool InputDelta(FEditorViewportClient* InViewportClient, FViewport* InViewport, FVector& InDrag, FRotator& InRot, FVector& InScale) override;
 	virtual bool StartTracking(FEditorViewportClient* InViewportClient, FViewport* InViewport) override;
 	virtual bool EndTracking(FEditorViewportClient* InViewportClient, FViewport* InViewport) override;
-	virtual bool RequiresLegacyViewportInteractions() const override { return true; }
+	virtual bool BeginTransform(const FGizmoState& InState) override;
+	virtual bool EndTransform(const FGizmoState& InState) override;
+	virtual bool RequiresLegacyViewportInteractions() const override { return false; }
 
 	// ILegacyEdModeWidgetInterface
 	virtual bool UsesToolkits() const override { return false; }
