Sequencer: Anim Mixer Force Root Lock refactor.
Instead of handling root lock inside the skeletal animation evaluation task (which makes it impossible to generalize in the root motion extraction task), pass through the force root lock behavior so the root motion extraction task can handle it instead. This allows us to remove the hack in the evaluation task that was early out'ing if it discovered the pose generation evaluation task had already written to the root motion attribute. The control rig pose generation task can have data on that attribute from passthrough, and we don't want to cause that to never evaluate root motion from control rig.

#Synced-CL 51964473

diff --git a/Engine/Plugins/Experimental/MovieSceneAnimMixer/Source/MovieSceneAnimMixer/Private/Systems/MovieSceneMixedSkeletalAnimationSystem.cpp b/Engine/Plugins/Experimental/MovieSceneAnimMixer/Source/MovieSceneAnimMixer/Private/Systems/MovieSceneMixedSkeletalAnimationSystem.cpp
--- a/Engine/Plugins/Experimental/MovieSceneAnimMixer/Source/MovieSceneAnimMixer/Private/Systems/MovieSceneMixedSkeletalAnimationSystem.cpp
+++ b/Engine/Plugins/Experimental/MovieSceneAnimMixer/Source/MovieSceneAnimMixer/Private/Systems/MovieSceneMixedSkeletalAnimationSystem.cpp
@@ -186,6 +186,11 @@ struct FSkeletalAnimMixerMutation : IMovieSceneEntityMutation
 						OutRootMotionSettings[Index].bHasRootMotionOverride = true;
 					}
 				}
+
+				if (UAnimSequence* AnimSequence = Cast<UAnimSequence>(SkeletalAnim.Section ? SkeletalAnim.Section->Params.Animation.Get() : nullptr))
+				{
+					OutRootMotionSettings[Index].bForceRootLock = AnimSequence->bForceRootLock && AnimSequence->bEnableRootMotion;
+				}
 			}
 		}
 	}
@@ -360,6 +365,7 @@ struct FGatherMixableSkeletalAnimations
 			AnimData.bFireNotifies = bFireNotifies && !AnimParams.bSkipAnimNotifiers && !bLooped;
 			// TODO: We want to allow for additive to be set on the section properties
 			AnimData.bAdditive = AnimData.AnimSequence.IsValid() ? AnimData.AnimSequence->IsValidAdditive() : false;
+			AnimData.bSuppressForceRootLock = OptionalRootMotionSettings && OptionalRootMotionSettings[Index].bForceRootLock;
 
 			// TODO: Need to create a mirroring task and add that if appropriate
 			// TODO: Notifies- perhaps this should be a part of the animnext side of things?
@@ -426,6 +432,11 @@ void FMovieSceneSkeletalAnimationEvaluationTask::Execute(UE::UAF::FEvaluationVM&
 
 		const bool bExtractRootMotion = AnimationData.RootMotionOverride.IsSet();
 		FAnimExtractContext ExtractionContext(AnimationData.ToPosition, bExtractRootMotion, DeltaTime, false);
+#if WITH_EDITOR
+		// When the extraction task will handle root motion, suppress Force Root Lock
+		// during pose evaluation so the root bone retains its actual motion.
+		ExtractionContext.bIgnoreRootLock = AnimationData.bSuppressForceRootLock;
+#endif
 
 		FKeyframeState Keyframe = VM.MakeUninitializedKeyframe(bIsAdditive);
 
@@ -488,31 +499,6 @@ void FMovieSceneSkeletalAnimationEvaluationTask::Execute(UE::UAF::FEvaluationVM&
 					}
 				}
 			}
-#if WITH_EDITOR
-			// When bForceRootLock is true, the decompressor locks bone 0 to the
-			// reference pose so the actual root motion is not in the pose data.
-			// Pre-extract it from the raw track data so FAnimNextExtractRootMotionTask
-			// doesn't need access to the UAnimSequence. This only affects bone 0 as
-			// child bones are not locked, so FirstAnimatedChild extraction from the
-			// pose still works correctly.
-			else if (AnimSequencePtr->bForceRootLock && AnimSequencePtr->bEnableRootMotion)
-			{
-				FAnimExtractContext ExtractContext(AnimationData.ToPosition);
-				ExtractContext.bIgnoreRootLock = true;
-				FTransform RootTrackTransform = AnimSequencePtr->ExtractRootTrackTransform(ExtractContext, nullptr);
-
-				FTransformAnimationAttribute* RootMotionAttribute = Keyframe.Attributes.FindOrAdd<FTransformAnimationAttribute>(AnimMixerComponents->RootTransformAttributeId);
-				FFloatAnimationAttribute* RootMotionWeight = Keyframe.Attributes.FindOrAdd<FFloatAnimationAttribute>(AnimMixerComponents->RootTransformWeightAttributeId);
-				RootMotionAttribute->Value = RootTrackTransform;
-				RootMotionWeight->Value = 1.f;
-
-				const int32 RootIndex = Keyframe.Pose.GetRefPose().GetLODBoneIndexFromSkeletonBoneIndex(0);
-				if (RootIndex != INDEX_NONE)
-				{
-					Keyframe.Pose.LocalTransformsView[RootIndex] = FTransform::Identity;
-				}
-			}
-#endif
 		}
 
 		if (EnumHasAnyFlags(VM.GetFlags(), EEvaluationFlags::Curves))
diff --git a/Engine/Plugins/Experimental/MovieSceneAnimMixer/Source/MovieSceneAnimMixer/Private/Systems/MovieSceneRootMotionSystem.cpp b/Engine/Plugins/Experimental/MovieSceneAnimMixer/Source/MovieSceneAnimMixer/Private/Systems/MovieSceneRootMotionSystem.cpp
--- a/Engine/Plugins/Experimental/MovieSceneAnimMixer/Source/MovieSceneAnimMixer/Private/Systems/MovieSceneRootMotionSystem.cpp
+++ b/Engine/Plugins/Experimental/MovieSceneAnimMixer/Source/MovieSceneAnimMixer/Private/Systems/MovieSceneRootMotionSystem.cpp
@@ -621,13 +621,6 @@ void FAnimNextExtractRootMotionTask::Execute(UE::UAF::FEvaluationVM& VM) const
 
 	const FAnimMixerComponentTypes* AnimMixerComponents = FAnimMixerComponentTypes::Get();
 
-	// If a RootTransformAttribute already exists (e.g. from a legacy RootMotionOverride),
-	// don't overwrite it.
-	if ((*Keyframe)->Attributes.Find<FTransformAnimationAttribute>(AnimMixerComponents->RootTransformAttributeId))
-	{
-		return;
-	}
-
 	const int32 RootLODIndex = (*Keyframe)->Pose.GetRefPose().GetLODBoneIndexFromSkeletonBoneIndex(0);
 	if (RootLODIndex == INDEX_NONE)
 	{
diff --git a/Engine/Plugins/Experimental/MovieSceneAnimMixer/Source/MovieSceneAnimMixer/Public/Systems/MovieSceneMixedSkeletalAnimationSystem.h b/Engine/Plugins/Experimental/MovieSceneAnimMixer/Source/MovieSceneAnimMixer/Public/Systems/MovieSceneMixedSkeletalAnimationSystem.h
--- a/Engine/Plugins/Experimental/MovieSceneAnimMixer/Source/MovieSceneAnimMixer/Public/Systems/MovieSceneMixedSkeletalAnimationSystem.h
+++ b/Engine/Plugins/Experimental/MovieSceneAnimMixer/Source/MovieSceneAnimMixer/Public/Systems/MovieSceneMixedSkeletalAnimationSystem.h
@@ -38,6 +38,9 @@ struct FMixedAnimSkeletalAnimationData
 
 	UPROPERTY()
 	bool bAdditive = false;
+
+	UPROPERTY()
+	bool bSuppressForceRootLock = false;
 };
 
 USTRUCT()
diff --git a/Engine/Plugins/Experimental/MovieSceneAnimMixer/Source/MovieSceneAnimMixer/Public/Systems/MovieSceneRootMotionSystem.h b/Engine/Plugins/Experimental/MovieSceneAnimMixer/Source/MovieSceneAnimMixer/Public/Systems/MovieSceneRootMotionSystem.h
--- a/Engine/Plugins/Experimental/MovieSceneAnimMixer/Source/MovieSceneAnimMixer/Public/Systems/MovieSceneRootMotionSystem.h
+++ b/Engine/Plugins/Experimental/MovieSceneAnimMixer/Source/MovieSceneAnimMixer/Public/Systems/MovieSceneRootMotionSystem.h
@@ -94,6 +94,9 @@ struct FMovieSceneRootMotionSettings
 	uint8 bHasRootMotionOverride : 1 = 0;
 	/** If true we use the ChildBoneIndex otherwise we use the root*/
 	uint8 bBlendFirstChildOfRoot : 1 = 0;
+	/** If true, the animation has Force Root Lock + Enable Root Motion. The extraction task
+	 *  handles root locking, and the eval task suppresses the lock during pose evaluation. */
+	uint8 bForceRootLock : 1 = 0;
 };
 
 // Structure that is shared between entities for handling the mixer's root motion result. 
