Sequencer: Anim mixer- move offset channel evaluation from the skel anim eval task to the anim mixer system generally, which allows this to be used by other types of sections/tracks (e.g. Control Rig).

#Synced-CL 52256411

diff --git a/Engine/Plugins/Experimental/MovieSceneAnimMixer/Source/MovieSceneAnimMixer/Private/Systems/MovieSceneAnimMixerSystem.cpp b/Engine/Plugins/Experimental/MovieSceneAnimMixer/Source/MovieSceneAnimMixer/Private/Systems/MovieSceneAnimMixerSystem.cpp
--- a/Engine/Plugins/Experimental/MovieSceneAnimMixer/Source/MovieSceneAnimMixer/Private/Systems/MovieSceneAnimMixerSystem.cpp
+++ b/Engine/Plugins/Experimental/MovieSceneAnimMixer/Source/MovieSceneAnimMixer/Private/Systems/MovieSceneAnimMixerSystem.cpp
@@ -348,11 +348,7 @@ static TSharedPtr<UE::UAF::FEvaluationProgram> BuildProgramFromEntries(
 			}
 
 			bool bNeedsMirror = (WrapEntry->MirrorTable != nullptr);
-			// Skip extraction when a legacy root motion override is active -
-			// Execute() handles the override inline and no separate extraction
-			// from the pose is needed.
-			bool bNeedsExtraction = WrapEntry->RootMotionSettings.IsSet()
-				&& !WrapEntry->RootMotionSettings->bHasRootMotionOverride;
+			bool bNeedsExtraction = WrapEntry->RootMotionSettings.IsSet();
 			TOptional<FAnimNextConvertRootMotionToWorldSpaceTask> SpaceTask;
 			if (!Filter.bSkipRootMotionConversion)
 			{
@@ -543,10 +539,8 @@ static TSharedPtr<UE::UAF::FEvaluationProgram> BuildProgramFromEntries(
 
 			// Extract root motion from the pose into a RootTransformAttribute.
 			// This runs after eval and mirror, before space conversion.
-			// Skipped for transitions (already handled in mini-programs) and
-			// for legacy root motion overrides (handled inline by Execute()).
-			if (MixerEntry->RootMotionSettings.IsSet() && !MixerEntry->bIsTransition
-				&& !MixerEntry->RootMotionSettings->bHasRootMotionOverride)
+			// Skipped for transitions (already handled in mini-programs).
+			if (MixerEntry->RootMotionSettings.IsSet() && !MixerEntry->bIsTransition)
 			{
 				Program->AppendTask(FAnimNextExtractRootMotionTask::Make(
 					MixerEntry->RootMotionSettings->RootMotionSource));
@@ -625,7 +619,13 @@ namespace UE::MovieScene
 			TSharedPtr<FMovieSceneAnimMixerEntry> MixerEntry,
 			const FMovieSceneRootMotionSettings* RootMotionSettings,
 			const double* WeightAndEasings,
-			const FTransform* BoneMatchTransform)
+			const FTransform* BoneMatchTransform,
+			const double* OffsetX,
+			const double* OffsetY,
+			const double* OffsetZ,
+			const double* RotOffsetX,
+			const double* RotOffsetY,
+			const double* RotOffsetZ)
 		{
 			if (MixerEntry)
 			{
@@ -637,6 +637,18 @@ namespace UE::MovieScene
 				if (RootMotionSettings)
 				{
 					MixerEntry->RootMotionSettings = *RootMotionSettings;
+
+					// Apply evaluated offset channel results to the root motion settings
+					FVector Location = RootMotionSettings->RootLocation;
+					FRotator Rotation = RootMotionSettings->RootRotation.Rotator();
+					if (OffsetX)    { Location.X = *OffsetX; }
+					if (OffsetY)    { Location.Y = *OffsetY; }
+					if (OffsetZ)    { Location.Z = *OffsetZ; }
+					if (RotOffsetX) { Rotation.Roll  = *RotOffsetX; }
+					if (RotOffsetY) { Rotation.Pitch = *RotOffsetY; }
+					if (RotOffsetZ) { Rotation.Yaw   = *RotOffsetZ; }
+					MixerEntry->RootMotionSettings->RootLocation = Location;
+					MixerEntry->RootMotionSettings->RootRotation = Rotation.Quaternion();
 				}
 
 				if (BoneMatchTransform)
@@ -1573,6 +1585,12 @@ void UMovieSceneAnimMixerSystem::OnSchedulePersistentTasks(UE::MovieScene::IEnti
 		.ReadOptional(AnimMixerComponents->RootMotionSettings)
 		.ReadOptional(BuiltInComponents->WeightAndEasingResult)
 		.ReadOptional(AnimMixerComponents->BoneMatchTransform)
+		.ReadOptional(BuiltInComponents->DoubleResult[0])
+		.ReadOptional(BuiltInComponents->DoubleResult[1])
+		.ReadOptional(BuiltInComponents->DoubleResult[2])
+		.ReadOptional(BuiltInComponents->DoubleResult[3])
+		.ReadOptional(BuiltInComponents->DoubleResult[4])
+		.ReadOptional(BuiltInComponents->DoubleResult[5])
 		.FilterNone({BuiltInComponents->Tags.NeedsUnlink, BuiltInComponents->Tags.Ignored})
 		.Schedule_PerEntity<FUpdateTaskEntities>(&Linker->EntityManager, TaskScheduler);
 
diff --git a/Engine/Plugins/Experimental/MovieSceneAnimMixer/Source/MovieSceneAnimMixer/Private/Systems/MovieSceneMixedSkeletalAnimationSystem.cpp b/Engine/Plugins/Experimental/MovieSceneAnimMixer/Source/MovieSceneAnimMixer/Private/Systems/MovieSceneMixedSkeletalAnimationSystem.cpp
--- a/Engine/Plugins/Experimental/MovieSceneAnimMixer/Source/MovieSceneAnimMixer/Private/Systems/MovieSceneMixedSkeletalAnimationSystem.cpp
+++ b/Engine/Plugins/Experimental/MovieSceneAnimMixer/Source/MovieSceneAnimMixer/Private/Systems/MovieSceneMixedSkeletalAnimationSystem.cpp
@@ -244,8 +244,6 @@ struct FGatherMixableSkeletalAnimations
 		TRead<FMovieSceneSkeletalAnimationComponentData> SkeletalAnimations,
 		TReadOptional<FFrameTime> OptionalEvalTimes,
 		TWriteOptional<FMovieSceneRootMotionSettings> OptionalRootMotionSettings,
-		const double* OptLocationX, const double* OptLocationY, const double* OptLocationZ,
-		const double* OptRotationX, const double* OptRotationY, const double* OptRotationZ,
 		TWrite<TSharedPtr<FAnimNextEvaluationTask>> EvalTask) const
 	{
 		// Gather all the skeletal animations currently active in all sequences.
@@ -377,22 +375,6 @@ struct FGatherMixableSkeletalAnimations
 			}
 			TSharedPtr<FMovieSceneSkeletalAnimationEvaluationTask> AnimTask = StaticCastSharedPtr<FMovieSceneSkeletalAnimationEvaluationTask>(EvalTask[Index]);
 			AnimTask->AnimationData = AnimData;
-
-			if (OptionalRootMotionSettings)
-			{
-				FVector  CurrentLocation = OptionalRootMotionSettings[Index].RootLocation;
-				FRotator CurrentRotation = OptionalRootMotionSettings[Index].RootRotation.Rotator();
-
-				if (OptLocationX) { CurrentLocation.X = OptLocationX[Index]; }
-				if (OptLocationY) { CurrentLocation.Y = OptLocationY[Index]; }
-				if (OptLocationZ) { CurrentLocation.Z = OptLocationZ[Index]; }
-				if (OptRotationX) { CurrentRotation.Roll  = OptRotationX[Index]; }
-				if (OptRotationY) { CurrentRotation.Pitch = OptRotationY[Index]; }
-				if (OptRotationZ) { CurrentRotation.Yaw   = OptRotationZ[Index]; }
-
-				OptionalRootMotionSettings[Index].RootLocation = CurrentLocation;
-				OptionalRootMotionSettings[Index].RootRotation = CurrentRotation.Quaternion();
-			}
 		}
 	}
 
@@ -610,12 +592,6 @@ void UMovieSceneMixedSkeletalAnimationSystem::OnSchedulePersistentTasks(UE::Movi
 		.Read(TrackComponents->SkeletalAnimation)
 		.ReadOptional(BuiltInComponents->EvalTime)
 		.WriteOptional(AnimMixerComponents->RootMotionSettings)
-		.ReadOptional(BuiltInComponents->DoubleResult[0]) // Root Offset X
-		.ReadOptional(BuiltInComponents->DoubleResult[1]) // Root Offset Y
-		.ReadOptional(BuiltInComponents->DoubleResult[2]) // Root Offset Z
-		.ReadOptional(BuiltInComponents->DoubleResult[3]) // Root Rotation Offset X
-		.ReadOptional(BuiltInComponents->DoubleResult[4]) // Root Rotation Offset Y
-		.ReadOptional(BuiltInComponents->DoubleResult[5]) // Root Rotation Offset Z
 		.Write(AnimMixerComponents->Task)
 		.FilterAll( { TrackComponents->Tags.AnimMixerPoseProducer })
 		.FilterNone({ BuiltInComponents->Tags.Ignored})
