Sequencer: Auto-show root motion gizmo when selecting mixer sections with root motion decorations.

#Synced-CL 52256411

diff --git a/Engine/Plugins/Experimental/MovieSceneAnimMixer/Source/MovieSceneAnimMixerEditor/Private/AnimationMixerTrackEditMode.cpp b/Engine/Plugins/Experimental/MovieSceneAnimMixer/Source/MovieSceneAnimMixerEditor/Private/AnimationMixerTrackEditMode.cpp
--- a/Engine/Plugins/Experimental/MovieSceneAnimMixer/Source/MovieSceneAnimMixerEditor/Private/AnimationMixerTrackEditMode.cpp
+++ b/Engine/Plugins/Experimental/MovieSceneAnimMixer/Source/MovieSceneAnimMixerEditor/Private/AnimationMixerTrackEditMode.cpp
@@ -94,8 +94,31 @@ void UAnimationMixerTrackEditMode::Enter()
 	GLevelEditorModeTools().ActivateMode(USequencerToolsEditMode::ModeName);
 }
 
+void UAnimationMixerTrackEditMode::SetSequencer(const TSharedPtr<ISequencer>& InSequencer)
+{
+	// Unsubscribe from previous sequencer if any
+	if (TSharedPtr<ISequencer> OldSequencer = WeakSequencer.Pin())
+	{
+		OldSequencer->GetSelectionChangedSections().Remove(SectionSelectionChangedHandle);
+		SectionSelectionChangedHandle.Reset();
+	}
+
+	WeakSequencer = InSequencer;
+
+	if (InSequencer)
+	{
+		SectionSelectionChangedHandle = InSequencer->GetSelectionChangedSections().AddUObject(
+			this, &UAnimationMixerTrackEditMode::OnSequencerSectionSelectionChanged);
+	}
+}
+
 void UAnimationMixerTrackEditMode::Exit()
 {
+	if (TSharedPtr<ISequencer> Sequencer = WeakSequencer.Pin())
+	{
+		Sequencer->GetSelectionChangedSections().Remove(SectionSelectionChangedHandle);
+		SectionSelectionChangedHandle.Reset();
+	}
 
 	if (bIsTransacting)
 	{
@@ -213,7 +236,25 @@ void UAnimationMixerTrackEditMode::Render(const FSceneView* View, FViewport* Vie
 					}
 
 					UMovieSceneRootMotionSettingsDecoration* Decoration = Section->FindDecoration<UMovieSceneRootMotionSettingsDecoration>();
-					if (!Decoration || !Decoration->bShowSkeleton)
+					if (!Decoration)
+					{
+						++SectionIndex;
+						continue;
+					}
+
+					// Check if this section is in SelectedRootData (for gizmo positioning)
+					bool bIsSelected = false;
+					for (const FMixerSelectedRootData& SelData : SelectedRootData)
+					{
+						if (SelData.Section.Get() == Section)
+						{
+							bIsSelected = true;
+							break;
+						}
+					}
+
+					// Skip sections that aren't showing skeletons and aren't selected
+					if (!Decoration->bShowSkeleton && !bIsSelected)
 					{
 						++SectionIndex;
 						continue;
@@ -240,66 +281,67 @@ void UAnimationMixerTrackEditMode::Render(const FSceneView* View, FViewport* Vie
 
 						FVector RootLocation = CachedPose.RootBoneLocation;
 
-						// Draw bones using SkeletalDebugRendering for proper bone shapes
-						const USkeletalMesh* SkelMesh = SkelMeshComp->GetSkeletalMeshAsset();
-						if (SkelMesh && CachedPose.WorldBoneTransforms.Num() > 0 && CachedPose.ParentBoneIndices.Num() > 0)
+						// Update cached root transform for gizmo positioning regardless
+						// of skeleton visibility
+						if (bIsSelected)
 						{
-							const int32 NumBones = CachedPose.WorldBoneTransforms.Num();
-
-							FSkelDebugDrawConfig DrawConfig;
-							DrawConfig.BoneDrawMode = EBoneDrawMode::All;
-							DrawConfig.BoneDrawSize = 1.f;
-							DrawConfig.bAddHitProxy = false;
-							DrawConfig.bForceDraw = true;
-							DrawConfig.DefaultBoneColor = BoneColor;
-							DrawConfig.AffectedBoneColor = BoneColor;
-							DrawConfig.SelectedBoneColor = BoneColor;
-							DrawConfig.ParentOfSelectedBoneColor = BoneColor;
-
-							TArray<uint16> RequiredBones;
-							RequiredBones.Reserve(NumBones);
-							TArray<FLinearColor> BoneColours;
-							BoneColours.Reserve(NumBones);
-							for (int32 BoneIdx = 0; BoneIdx < NumBones; ++BoneIdx)
-							{
-								RequiredBones.Add(static_cast<uint16>(BoneIdx));
-								BoneColours.Add(BoneColor);
-							}
-
-							SkeletalDebugRendering::DrawBones(
-								PDI,
-								SkelMeshComp->GetComponentLocation(),
-								RequiredBones,
-								SkelMesh->GetRefSkeleton(),
-								CachedPose.WorldBoneTransforms,
-								TArray<int32>(),
-								BoneColours,
-								TArray<TRefCountPtr<HHitProxy>>(),
-								DrawConfig
-							);
+							CurrentRootTransform.SetLocation(RootLocation);
 						}
 
-						// Draw root bone point with hit proxy for gizmo selection.
-						// Use a brighter color and larger size so it's easy to spot and click.
-						FLinearColor RootColor = (BoneColor + FLinearColor::White) * 0.5f;
-						RootColor.A = 1.0f;
-
-						if (PDI->IsHitTesting())
-						{
-							PDI->SetHitProxy(new HMixerEditModeRootHitProxy(Section, SkelMeshComp));
-						}
-						PDI->DrawPoint(RootLocation, RootColor, 18.f, SDPG_Foreground);
-						if (PDI->IsHitTesting())
+						if (Decoration->bShowSkeleton)
 						{
-							PDI->SetHitProxy(nullptr);
-						}
+							// Draw bones using SkeletalDebugRendering for proper bone shapes
+							const USkeletalMesh* SkelMesh = SkelMeshComp->GetSkeletalMeshAsset();
+							if (SkelMesh && CachedPose.WorldBoneTransforms.Num() > 0 && CachedPose.ParentBoneIndices.Num() > 0)
+							{
+								const int32 NumBones = CachedPose.WorldBoneTransforms.Num();
+
+								FSkelDebugDrawConfig DrawConfig;
+								DrawConfig.BoneDrawMode = EBoneDrawMode::All;
+								DrawConfig.BoneDrawSize = 1.f;
+								DrawConfig.bAddHitProxy = false;
+								DrawConfig.bForceDraw = true;
+								DrawConfig.DefaultBoneColor = BoneColor;
+								DrawConfig.AffectedBoneColor = BoneColor;
+								DrawConfig.SelectedBoneColor = BoneColor;
+								DrawConfig.ParentOfSelectedBoneColor = BoneColor;
+
+								TArray<uint16> RequiredBones;
+								RequiredBones.Reserve(NumBones);
+								TArray<FLinearColor> BoneColours;
+								BoneColours.Reserve(NumBones);
+								for (int32 BoneIdx = 0; BoneIdx < NumBones; ++BoneIdx)
+								{
+									RequiredBones.Add(static_cast<uint16>(BoneIdx));
+									BoneColours.Add(BoneColor);
+								}
+
+								SkeletalDebugRendering::DrawBones(
+									PDI,
+									SkelMeshComp->GetComponentLocation(),
+									RequiredBones,
+									SkelMesh->GetRefSkeleton(),
+									CachedPose.WorldBoneTransforms,
+									TArray<int32>(),
+									BoneColours,
+									TArray<TRefCountPtr<HHitProxy>>(),
+									DrawConfig
+								);
+							}
 
-						// Update cached root transform if this section is selected
-						for (const FMixerSelectedRootData& SelData : SelectedRootData)
-						{
-							if (SelData.Section.Get() == Section)
+							// Draw root bone point with hit proxy for gizmo selection.
+							// Use a brighter color and larger size so it's easy to spot and click.
+							FLinearColor RootColor = (BoneColor + FLinearColor::White) * 0.5f;
+							RootColor.A = 1.0f;
+
+							if (PDI->IsHitTesting())
 							{
-								CurrentRootTransform.SetLocation(RootLocation);
+								PDI->SetHitProxy(new HMixerEditModeRootHitProxy(Section, SkelMeshComp));
+							}
+							PDI->DrawPoint(RootLocation, RootColor, 18.f, SDPG_Foreground);
+							if (PDI->IsHitTesting())
+							{
+								PDI->SetHitProxy(nullptr);
 							}
 						}
 					}
@@ -615,6 +657,102 @@ bool UAnimationMixerTrackEditMode::InputDelta(FEditorViewportClient* InViewportC
 	return false;
 }
 
+USkeletalMeshComponent* UAnimationMixerTrackEditMode::FindSkelMeshForSection(UMovieSceneSection* Section) const
+{
+	TSharedPtr<ISequencer> Sequencer = WeakSequencer.Pin();
+	if (!Sequencer || !Section)
+	{
+		return nullptr;
+	}
+
+	UMovieSceneSequence* Sequence = Sequencer->GetFocusedMovieSceneSequence();
+	if (!Sequence)
+	{
+		return nullptr;
+	}
+
+	const UMovieScene* MovieScene = Sequence->GetMovieScene();
+	if (!MovieScene)
+	{
+		return nullptr;
+	}
+
+	UMovieSceneTrack* OwnerTrack = Section->GetTypedOuter<UMovieSceneTrack>();
+	if (!OwnerTrack)
+	{
+		return nullptr;
+	}
+
+	for (const FMovieSceneBinding& Binding : MovieScene->GetBindings())
+	{
+		for (UMovieSceneTrack* Track : Binding.GetTracks())
+		{
+			if (Track == OwnerTrack)
+			{
+				TArrayView<TWeakObjectPtr<UObject>> BoundObjects = Sequencer->FindObjectsInCurrentSequence(Binding.GetObjectGuid());
+				for (TWeakObjectPtr<> WeakObj : BoundObjects)
+				{
+					if (UObject* Obj = WeakObj.Get())
+					{
+						if (USkeletalMeshComponent* SkelMesh = MovieSceneToolHelpers::AcquireSkeletalMeshFromObject(Obj))
+						{
+							return SkelMesh;
+						}
+					}
+				}
+				return nullptr;
+			}
+		}
+	}
+
+	return nullptr;
+}
+
+void UAnimationMixerTrackEditMode::OnSequencerSectionSelectionChanged(TArray<UMovieSceneSection*> SelectedSections)
+{
+	TSharedPtr<ISequencer> Sequencer = WeakSequencer.Pin();
+	if (!Sequencer)
+	{
+		return;
+	}
+
+	// Build the new selection from sections that have a root motion decoration
+	TArray<FMixerSelectedRootData> NewSelection;
+	for (UMovieSceneSection* Section : SelectedSections)
+	{
+		if (!Section || !Section->IsActive())
+		{
+			continue;
+		}
+
+		UMovieSceneRootMotionSettingsDecoration* Decoration = Section->FindDecoration<UMovieSceneRootMotionSettingsDecoration>();
+		if (!Decoration)
+		{
+			continue;
+		}
+
+		UMovieSceneAnimationMixerTrack* MixerTrack = Section->GetTypedOuter<UMovieSceneAnimationMixerTrack>();
+		if (!MixerTrack)
+		{
+			continue;
+		}
+
+		USkeletalMeshComponent* SkelMesh = FindSkelMeshForSection(Section);
+		if (!SkelMesh)
+		{
+			continue;
+		}
+
+		FMixerSelectedRootData RootData;
+		RootData.Section = Section;
+		RootData.MeshComp = SkelMesh;
+		RootData.Decoration = Decoration;
+		NewSelection.Add(MoveTemp(RootData));
+	}
+
+	SelectedRootData = MoveTemp(NewSelection);
+}
+
 void UAnimationMixerTrackEditMode::KeyOffsetChannels(
 	UMovieSceneRootMotionSettingsDecoration* Decoration,
 	const FVector& LocationOffset,
diff --git a/Engine/Plugins/Experimental/MovieSceneAnimMixer/Source/MovieSceneAnimMixerEditor/Private/AnimationMixerTrackEditMode.h b/Engine/Plugins/Experimental/MovieSceneAnimMixer/Source/MovieSceneAnimMixerEditor/Private/AnimationMixerTrackEditMode.h
--- a/Engine/Plugins/Experimental/MovieSceneAnimMixer/Source/MovieSceneAnimMixerEditor/Private/AnimationMixerTrackEditMode.h
+++ b/Engine/Plugins/Experimental/MovieSceneAnimMixer/Source/MovieSceneAnimMixerEditor/Private/AnimationMixerTrackEditMode.h
@@ -116,7 +116,7 @@ public:
 	virtual bool GetCustomDrawingCoordinateSystem(FMatrix& OutMatrix, void* InData) override;
 	virtual bool GetCustomInputCoordinateSystem(FMatrix& OutMatrix, void* InData) override;
 
-	void SetSequencer(const TSharedPtr<ISequencer>& InSequencer) { WeakSequencer = InSequencer; }
+	void SetSequencer(const TSharedPtr<ISequencer>& InSequencer);
 	void SetTrailHierarchy(const TSharedPtr<UE::Sequencer::FAnimMixerTrailHierarchy>& InHierarchy) { WeakTrailHierarchy = InHierarchy; }
 	void InvalidateAllCachedPoses();
 
@@ -136,6 +136,13 @@ private:
 	bool EndDragInternal();
 	void InvalidateOffsetCacheForSelection();
 
+	// Respond to sequencer section selection changes by auto-selecting roots
+	// for sections that have a root motion settings decoration
+	void OnSequencerSectionSelectionChanged(TArray<UMovieSceneSection*> SelectedSections);
+
+	// Find the skeletal mesh component bound to a section's track in the sequencer
+	USkeletalMeshComponent* FindSkelMeshForSection(UMovieSceneSection* Section) const;
+
 	// Cached skeleton pose for a single section, invalidated when scrub time or data changes
 	struct FCachedSectionPose
 	{
@@ -169,4 +176,6 @@ private:
 
 	bool bIsTransacting = false;
 	bool bManipulatorMadeChange = false;
+
+	FDelegateHandle SectionSelectionChangedHandle;
 };
