Fix shutdown issue with MovieScene

#rnx
#Synced-CL 52302118

diff --git a/Engine/Plugins/Experimental/MovieSceneAnimMixer/Source/MovieSceneAnimMixer/Private/MovieSceneAnimMixerModule.cpp b/Engine/Plugins/Experimental/MovieSceneAnimMixer/Source/MovieSceneAnimMixer/Private/MovieSceneAnimMixerModule.cpp
--- a/Engine/Plugins/Experimental/MovieSceneAnimMixer/Source/MovieSceneAnimMixer/Private/MovieSceneAnimMixerModule.cpp
+++ b/Engine/Plugins/Experimental/MovieSceneAnimMixer/Source/MovieSceneAnimMixer/Private/MovieSceneAnimMixerModule.cpp
@@ -59,9 +59,12 @@ void FMovieSceneAnimMixerModule::StartupModule()
 
 void FMovieSceneAnimMixerModule::ShutdownModule()
 {
-	IMovieSceneModule::Get().UnregisterCompatibleDecoration(UMovieSceneAnimationMixerTrack::StaticClass(), UMovieSceneRootMotionTargetDecoration::StaticClass());
-	IMovieSceneModule::Get().UnregisterCompatibleDecoration(UMovieSceneSkeletalAnimationSection::StaticClass(), UMovieSceneRootMotionSettingsDecoration::StaticClass());
-	IMovieSceneModule::Get().UnregisterCompatibleDecoration(UMovieSceneAnimationMixerLayer::StaticClass(), UMovieSceneAnimationMaskDecoration::StaticClass());
+	if (IMovieSceneModule::IsAvailable() && UObjectInitialized())
+	{
+		IMovieSceneModule::Get().UnregisterCompatibleDecoration(UMovieSceneAnimationMixerTrack::StaticClass(), UMovieSceneRootMotionTargetDecoration::StaticClass());
+		IMovieSceneModule::Get().UnregisterCompatibleDecoration(UMovieSceneSkeletalAnimationSection::StaticClass(), UMovieSceneRootMotionSettingsDecoration::StaticClass());
+		IMovieSceneModule::Get().UnregisterCompatibleDecoration(UMovieSceneAnimationMixerLayer::StaticClass(), UMovieSceneAnimationMaskDecoration::StaticClass());
+	}
 
 	IModularFeatures::Get().UnregisterModularFeature(IAnimGraphRuntime_SequencerMixerTargetConnector::GetModularFeatureName(), this);
 }
diff --git a/Engine/Plugins/Experimental/MovieSceneMixedControlRig/Source/MovieSceneMixedControlRig/Private/MovieSceneMixedControlRigModule.cpp b/Engine/Plugins/Experimental/MovieSceneMixedControlRig/Source/MovieSceneMixedControlRig/Private/MovieSceneMixedControlRigModule.cpp
--- a/Engine/Plugins/Experimental/MovieSceneMixedControlRig/Source/MovieSceneMixedControlRig/Private/MovieSceneMixedControlRigModule.cpp
+++ b/Engine/Plugins/Experimental/MovieSceneMixedControlRig/Source/MovieSceneMixedControlRig/Private/MovieSceneMixedControlRigModule.cpp
@@ -28,7 +28,7 @@ void FMovieSceneMixedControlRigModule::StartupModule()
 
 void FMovieSceneMixedControlRigModule::ShutdownModule()
 {
-	if (IMovieSceneModule::IsAvailable())
+	if (IMovieSceneModule::IsAvailable() && UObjectInitialized())
 	{
 		IMovieSceneModule::Get().UnregisterCompatibleDecoration(
 			UMovieSceneControlRigParameterTrack::StaticClass(),
