Sequencer: Fix undo crash for mask section additions/removals

Add missing Modify() calls before mutating the MaskSections array in
AddSection and RemoveSection so the undo system records the before-state.
Enhance PostEditUndo to re-establish owning layer references on mask
sections after undo restores the array.

#Synced-CL 52294507

diff --git a/Engine/Plugins/Experimental/MovieSceneAnimMixer/Source/MovieSceneAnimMixer/Private/MovieSceneAnimationMaskDecoration.cpp b/Engine/Plugins/Experimental/MovieSceneAnimMixer/Source/MovieSceneAnimMixer/Private/MovieSceneAnimationMaskDecoration.cpp
--- a/Engine/Plugins/Experimental/MovieSceneAnimMixer/Source/MovieSceneAnimMixer/Private/MovieSceneAnimationMaskDecoration.cpp
+++ b/Engine/Plugins/Experimental/MovieSceneAnimMixer/Source/MovieSceneAnimMixer/Private/MovieSceneAnimationMaskDecoration.cpp
@@ -4,6 +4,7 @@
 
 #include "MovieScene.h"
 #include "MovieSceneAnimMixerMaskSection.h"
+#include "MovieSceneAnimationMixerLayer.h"
 #include "EvaluationVM/Tasks/BlendKeyframesPerBone.h"
 
 TSharedPtr<FAnimNextEvaluationTask> FMovieSceneMaskBlendProvider::GetBlendTask()
@@ -170,6 +171,8 @@ void UMovieSceneAnimationMaskDecoration::AddSection(UMovieSceneSection* InSectio
 		return;
 	}
 
+	Modify();
+
 	const TRange<FFrameNumber> NewSectionRange = NewMaskSection->GetRange();
 
 	if (NewSectionRange.GetLowerBound().IsClosed() && NewSectionRange.GetUpperBound().IsClosed())
@@ -207,6 +210,7 @@ void UMovieSceneAnimationMaskDecoration::AddSection(UMovieSceneSection* InSectio
 
 void UMovieSceneAnimationMaskDecoration::RemoveSection(UMovieSceneSection& SectionToRemove)
 {
+	Modify();
 	MaskSections.RemoveAll([&SectionToRemove](const UMovieSceneSection* InSection){ return InSection == &SectionToRemove; });
 	MarkAsChanged();
 #if WITH_EDITOR
@@ -218,6 +222,17 @@ void UMovieSceneAnimationMaskDecoration::RemoveSection(UMovieSceneSection& Secti
 void UMovieSceneAnimationMaskDecoration::PostEditUndo()
 {
 	Super::PostEditUndo();
+
+	// Re-establish owning layer references after undo restores the section array.
+	UMovieSceneAnimationMixerLayer* Layer = GetTypedOuter<UMovieSceneAnimationMixerLayer>();
+	for (TObjectPtr<UMovieSceneSection> Section : MaskSections)
+	{
+		if (UMovieSceneAnimMixerMaskSection* MaskSection = Cast<UMovieSceneAnimMixerMaskSection>(Section))
+		{
+			MaskSection->SetAnimMixerLayer(Layer);
+		}
+	}
+
 	MarkStructureChanged();
 }
 #endif
\ No newline at end of file
