Sequencer: Anim Mixer Masks- added a 'Create New Mask' option that auto-creates a mask, saves it, and opens the editor for it so you can edit it. This addresses the concerns in the jira.

#Synced-CL 52294507

diff --git a/Engine/Plugins/Experimental/MovieSceneAnimMixer/Source/MovieSceneAnimMixerEditor/Private/AnimMixerDecorationMenuProviders.cpp b/Engine/Plugins/Experimental/MovieSceneAnimMixer/Source/MovieSceneAnimMixerEditor/Private/AnimMixerDecorationMenuProviders.cpp
--- a/Engine/Plugins/Experimental/MovieSceneAnimMixer/Source/MovieSceneAnimMixerEditor/Private/AnimMixerDecorationMenuProviders.cpp
+++ b/Engine/Plugins/Experimental/MovieSceneAnimMixer/Source/MovieSceneAnimMixerEditor/Private/AnimMixerDecorationMenuProviders.cpp
@@ -3,6 +3,8 @@
 #include "AnimMixerDecorationMenuProviders.h"
 
 #include "AnimatedRange.h"
+#include "Subsystems/AssetEditorSubsystem.h"
+#include "AssetToolsModule.h"
 #include "ContentBrowserModule.h"
 #include "IContentBrowserSingleton.h"
 #include "ISequencer.h"
@@ -14,6 +16,9 @@
 #include "Components/SkeletalMeshComponent.h"
 #include "Decorations/MovieSceneDecorationContainer.h"
 #include "Engine/SkeletalMesh.h"
+#include "Editor.h"
+#include "FileHelpers.h"
+#include "Factories/Factory.h"
 #include "Framework/MultiBox/MultiBoxBuilder.h"
 #include "GameFramework/Actor.h"
 #include "UAF/BlendMask/UAFBlendMask.h"
@@ -63,6 +68,46 @@ const USkeleton* FMaskDecorationMenuProvider::FindSkeletonForObject(const TObjec
 	return SkeletalMesh && SkeletalMesh->GetSkeletalMeshAsset() ? SkeletalMesh->GetSkeletalMeshAsset()->GetSkeleton() : nullptr;
 }
 
+void FMaskDecorationMenuProvider::OnCreateNewBlendMask(UMovieSceneDecorationContainerObject* DecorationContainer, TWeakPtr<ISequencer> WeakSequencer)
+{
+	UFactory* BlendMaskFactory = nullptr;
+	for (TObjectIterator<UClass> It; It; ++It)
+	{
+		UClass* CurrentClass = *It;
+		if (CurrentClass->IsChildOf(UFactory::StaticClass()) && !CurrentClass->HasAnyClassFlags(CLASS_Abstract))
+		{
+			UFactory* Factory = Cast<UFactory>(CurrentClass->GetDefaultObject());
+			if (Factory->CanCreateNew() && Factory->SupportedClass == UUAFBlendMask::StaticClass())
+			{
+				BlendMaskFactory = Factory;
+				break;
+			}
+		}
+	}
+
+	if (!BlendMaskFactory)
+	{
+		return;
+	}
+
+	IAssetTools& AssetTools = FAssetToolsModule::GetModule().Get();
+	UObject* NewAsset = AssetTools.CreateAssetWithDialog(UUAFBlendMask::StaticClass(), BlendMaskFactory);
+	if (!NewAsset)
+	{
+		return;
+	}
+
+	// Save the new asset so it persists and is available in the asset registry for future picking.
+	constexpr bool bOnlyDirty = false;
+	UEditorLoadingAndSavingUtils::SavePackages({ NewAsset->GetPackage() }, bOnlyDirty);
+
+	// Open the blend mask editor so the user can configure it immediately.
+	GEditor->GetEditorSubsystem<UAssetEditorSubsystem>()->OpenEditorForAsset(NewAsset);
+
+	FAssetData AssetData(NewAsset);
+	OnMaskAssetSelected(AssetData, DecorationContainer, WeakSequencer);
+}
+
 void FMaskDecorationMenuProvider::PopulateMaskSubmenu(
 	FMenuBuilder& MenuBuilder,
 	TObjectPtr<UObject> BoundObject,
@@ -70,6 +115,15 @@ void FMaskDecorationMenuProvider::PopulateMaskSubmenu(
 	TWeakPtr<ISequencer> InSequencer
 	)
 {
+	MenuBuilder.AddMenuEntry(
+		LOCTEXT("CreateNewBlendMask", "Create New Blend Mask"),
+		LOCTEXT("CreateNewBlendMaskTooltip", "Create a new blend mask asset."),
+		FSlateIcon(),
+		FUIAction(FExecuteAction::CreateRaw(this, &FMaskDecorationMenuProvider::OnCreateNewBlendMask, DecorationContainer, InSequencer))
+	);
+
+	MenuBuilder.AddSeparator();
+
 	const USkeleton* Skeleton = FindSkeletonForObject(BoundObject);
 
 	FAssetPickerConfig AssetPickerConfig;
diff --git a/Engine/Plugins/Experimental/MovieSceneAnimMixer/Source/MovieSceneAnimMixerEditor/Private/AnimMixerDecorationMenuProviders.h b/Engine/Plugins/Experimental/MovieSceneAnimMixer/Source/MovieSceneAnimMixerEditor/Private/AnimMixerDecorationMenuProviders.h
--- a/Engine/Plugins/Experimental/MovieSceneAnimMixer/Source/MovieSceneAnimMixerEditor/Private/AnimMixerDecorationMenuProviders.h
+++ b/Engine/Plugins/Experimental/MovieSceneAnimMixer/Source/MovieSceneAnimMixerEditor/Private/AnimMixerDecorationMenuProviders.h
@@ -25,6 +25,8 @@ private:
 
 	void PopulateMaskSubmenu(FMenuBuilder& MenuBuilder, TObjectPtr<UObject> BoundObject, UMovieSceneDecorationContainerObject* DecorationContainer, TWeakPtr<ISequencer> InSequencer);
 
+	void OnCreateNewBlendMask(UMovieSceneDecorationContainerObject* DecorationContainer, TWeakPtr<ISequencer> WeakSequencer);
+
 	void OnMaskAssetSelected(const FAssetData& AssetData, UMovieSceneDecorationContainerObject* DecorationContainer, TWeakPtr<ISequencer> Sequencer);
 
 	bool FilterBlendMasks(const FAssetData& AssetData, const USkeleton* Skeleton);
