Sequencer: Normalize rotations after per-bone mask blend

#Synced-CL 52294507

diff --git a/Engine/Plugins/Experimental/MovieSceneAnimMixer/Source/MovieSceneAnimMixer/Private/Systems/MovieSceneAnimMixerSystem.cpp b/Engine/Plugins/Experimental/MovieSceneAnimMixer/Source/MovieSceneAnimMixer/Private/Systems/MovieSceneAnimMixerSystem.cpp
--- a/Engine/Plugins/Experimental/MovieSceneAnimMixer/Source/MovieSceneAnimMixer/Private/Systems/MovieSceneAnimMixerSystem.cpp
+++ b/Engine/Plugins/Experimental/MovieSceneAnimMixer/Source/MovieSceneAnimMixer/Private/Systems/MovieSceneAnimMixerSystem.cpp
@@ -588,6 +588,8 @@ static TSharedPtr<UE::UAF::FEvaluationProgram> BuildProgramFromEntries(
 					{
 						Program->AppendTaskPtr(MixerEntry->BlendTask.Pin());
 						MixerEntry->BlendTask = nullptr;
+						// Per-bone mask blends can leave rotations un-normalized.
+						Program->AppendTask(FAnimNextNormalizeKeyframeRotationsTask());
 					}
 					else
 					{
