Sequencer: Restructure mixer evaluation around layers; new mask blend system

1) Two-level evaluation in BuildProgramFromEntries

The mixer evaluation loop now separates intra-layer and inter-layer blending. Each priority group (layer) first evaluates its entries into a single pose on the stack - handling transitions, weighted averages, and additives internally - then blends that result with the layer below using a layer-level blend operation. Previously these two concerns were mixed together in a flat loop, which made it impossible for layer-level operations like masks to apply consistently across all cases (single sections, overlapping sections, and transitions).

Weight handling is also split: EasingResult controls intra-layer blending (crossfade alpha, section overlap), while WeightResult contributes to the inter-layer blend weight. This allows section weights to survive to the layer blending step while allowing transitions to work between sections on the same layer hijacking the easing values. The RequiresBlending tag is removed since every layer now always blends with what is below it.

2) Dedicated mask system (UMovieSceneAnimMixerMaskSystem)

Mask blend accumulation is moved out of the shared UObject decoration (FMovieSceneMaskBlendProvider on UMovieSceneAnimationMaskDecoration) into a new ECS system. The old approach stored runtime accumulation state on the decoration which is shared data - multiple instances of the same sequence playing on different actors would corrupt each other. The new system creates per-{BoundObject, MixerLayer} output entities with a LayerBlendTask builder that the mixer consumes generically. This also enables future blend types to plug in via the same contract without the mixer needing to know about them.

3) Also fixed decoration muting for layers- the section didn't have the right row index.

#Synced-CL 52466902

diff --git a/Engine/Plugins/Experimental/MovieSceneAnimMixer/Source/MovieSceneAnimMixer/Private/AnimMixerComponentTypes.cpp b/Engine/Plugins/Experimental/MovieSceneAnimMixer/Source/MovieSceneAnimMixer/Private/AnimMixerComponentTypes.cpp
--- a/Engine/Plugins/Experimental/MovieSceneAnimMixer/Source/MovieSceneAnimMixer/Private/AnimMixerComponentTypes.cpp
+++ b/Engine/Plugins/Experimental/MovieSceneAnimMixer/Source/MovieSceneAnimMixer/Private/AnimMixerComponentTypes.cpp
@@ -54,9 +54,9 @@ namespace UE::MovieScene
 		ComponentRegistry->NewComponentType(&BoneMatchTransform, TEXT("Bone Match Transform"));
 		ComponentRegistry->NewComponentType(&GapBehavior, TEXT("Gap Behavior"));
 		ComponentRegistry->NewComponentType(&EntityOwner, TEXT("Entity Owner"));
+		ComponentRegistry->NewComponentType(&LayerBlendTask, TEXT("Layer Blend Task"));
 
 
-		Tags.RequiresBlending = ComponentRegistry->NewTag(TEXT("Requires Blending"));
 		Tags.Transition = ComponentRegistry->NewTag(TEXT("Transition"));
 		Tags.SkipPoseWeight = ComponentRegistry->NewTag(TEXT("Skip Pose Weight"));
 
@@ -79,7 +79,6 @@ namespace UE::MovieScene
 		ComponentRegistry->Factories.DefineChildComponent(RootDestination, MixerRootMotion);
 		ComponentRegistry->Factories.DefineChildComponent(Task, MixerEntry);
 
-		ComponentRegistry->Factories.DefineChildComponent(Tags.RequiresBlending, Tags.RequiresBlending);
 		ComponentRegistry->Factories.DefineChildComponent(Tags.Transition, Tags.Transition);
 
 	}
diff --git a/Engine/Plugins/Experimental/MovieSceneAnimMixer/Source/MovieSceneAnimMixer/Private/MovieSceneAnimMixerMaskSection.cpp b/Engine/Plugins/Experimental/MovieSceneAnimMixer/Source/MovieSceneAnimMixer/Private/MovieSceneAnimMixerMaskSection.cpp
--- a/Engine/Plugins/Experimental/MovieSceneAnimMixer/Source/MovieSceneAnimMixer/Private/MovieSceneAnimMixerMaskSection.cpp
+++ b/Engine/Plugins/Experimental/MovieSceneAnimMixer/Source/MovieSceneAnimMixer/Private/MovieSceneAnimMixerMaskSection.cpp
@@ -6,6 +6,7 @@
 #include "MovieSceneAnimationMaskDecoration.h"
 #include "MovieSceneAnimationMixerTrack.h"
 #include "EntitySystem/BuiltInComponentTypes.h"
+#include "Systems/MovieSceneSkeletalAnimationSystem.h"
 #include "MaskProfile/HierarchyTableTypeMask.h"
 
 #include UE_INLINE_GENERATED_CPP_BY_NAME(MovieSceneAnimMixerMaskSection)
@@ -33,7 +34,9 @@ void UMovieSceneAnimMixerMaskSection::ImportEntityImpl(UMovieSceneEntitySystemLi
 {
 	using namespace UE::MovieScene;
 
-	const FAnimMixerComponentTypes* AnimMixerComponentTypes = FAnimMixerComponentTypes::Get();
+	const FAnimMixerComponentTypes* AnimMixerComponentTypes = FAnimMixerComponentTypes::Get(); 
+	FBuiltInComponentTypes* BuiltInComponentTypes = FBuiltInComponentTypes::Get();
+
 
 	BuildDefaultComponents(EntityLinker, Params, OutImportedEntity);
 	
@@ -49,6 +52,8 @@ void UMovieSceneAnimMixerMaskSection::ImportEntityImpl(UMovieSceneEntitySystemLi
 
 	OutImportedEntity->AddBuilder(
 	FEntityBuilder()
+	.Add(BuiltInComponentTypes->GenericObjectBinding, Params.GetObjectBindingID())
+	.Add(BuiltInComponentTypes->BoundObjectResolver, UMovieSceneSkeletalAnimationSystem::ResolveSkeletalMeshComponentBinding)
 	.Add(AnimMixerComponentTypes->MixerLayer, OwningLayer)
 	.Add(AnimMixerComponentTypes->BlendProviderData, 	TInstancedStruct<FMovieSceneMaskBlendProviderData>::Make(MoveTemp(BlendData)))
 	.AddMutualComponents());
diff --git a/Engine/Plugins/Experimental/MovieSceneAnimMixer/Source/MovieSceneAnimMixer/Private/MovieSceneAnimationMaskDecoration.cpp b/Engine/Plugins/Experimental/MovieSceneAnimMixer/Source/MovieSceneAnimMixer/Private/MovieSceneAnimationMaskDecoration.cpp
--- a/Engine/Plugins/Experimental/MovieSceneAnimMixer/Source/MovieSceneAnimMixer/Private/MovieSceneAnimationMaskDecoration.cpp
+++ b/Engine/Plugins/Experimental/MovieSceneAnimMixer/Source/MovieSceneAnimMixer/Private/MovieSceneAnimationMaskDecoration.cpp
@@ -5,152 +5,6 @@
 #include "MovieScene.h"
 #include "MovieSceneAnimMixerMaskSection.h"
 #include "MovieSceneAnimationMixerLayer.h"
-#include "EvaluationVM/Tasks/BlendKeyframesPerBone.h"
-
-TSharedPtr<FAnimNextEvaluationTask> FMovieSceneMaskBlendProvider::GetBlendTask()
-{
-	BlendTask = nullptr;
-	if (AccumulationBoneWeights.IsEmpty())
-	{
-		ResetAccumulationData();
-		return BlendTask;
-	}
-
-	// Curve data needs to be re-allocated due to the nature of TNamedValueArray.
-	// @TODO This may be optimized by checking if the data has changed, and not emptying the weights if it hasn't, which should be the case more often than not.
-	TaskCurveWeights.Empty();
-	TaskAttributeWeights.Reset();
-	TaskBoneWeights.Reset();
-
-	// Second pass of calculating bone weights. Previously the damping values were accumulated, but they need to be inverted to get the weight again.
-	for (int32 WeightIndex = 0; WeightIndex < AccumulationBoneWeights.Num(); ++WeightIndex)
-	{
-		// Final Weight = 1.0f - Accumulated Damping Values. 
-		TaskBoneWeights.Add(1.0f - AccumulationBoneWeights[WeightIndex]);
-	}
-
-	for (const TTuple<UE::Anim::FAttributeId, float> CurrentAttributeWeight : AttributeWeightsMap)
-	{
-		TaskAttributeWeights.Add(FBlendProfileStandaloneCachedData::FMaskedAttributeWeight(CurrentAttributeWeight.Key, 1.0 - CurrentAttributeWeight.Value));
-	}
-
-	for (const TTuple<FName, float> CurrentCurveWeight : CurveWeightsMap)
-	{
-		TaskCurveWeights.Add(UE::Anim::FCurveElement(CurrentCurveWeight.Key, 1.0 - CurrentCurveWeight.Value));
-	}
-
-	BlendTask = MakeShared<FAnimNextBlendKeyframePerBoneWithScaleTask>(
-	FAnimNextBlendKeyframePerBoneWithScaleTask::Make(MaskSkeleton.Get(), &TaskBoneWeights,
-	&TaskCurveWeights, &TaskAttributeWeights, 1.0f));
-
-	ResetAccumulationData();
-	return BlendTask;
-}
-
-void FMovieSceneMaskBlendProvider::RegisterTaskData(const TInstancedStruct<FMovieSceneAnimMixerBlendProviderData>& InData, const double Weight)
-{
-	if (const FMovieSceneMaskBlendProviderData* MaskData = InData.GetPtr<FMovieSceneMaskBlendProviderData>())
-	{
-		if (MaskData->BlendMask.IsValid())
-		{
-			RegisterMaskAndWeight(MaskData->BlendMask.Get(), Weight);
-		}
-	}
-}
-
-void FMovieSceneMaskBlendProvider::RegisterMaskAndWeight(TObjectPtr<UUAFBlendMask> BlendMask, const double EasingWeight)
-{
-	if (!BlendMask)
-	{
-		return;
-	}
-
-	if (!MaskSkeleton.IsValid())
-	{
-		MaskSkeleton = BlendMask->GetSkeleton();
-	}
-
-	const TArray<float>& AssetBoneWeights = BlendMask->GetSkeletonBoneWeights();
-	const TArray<FBlendProfileStandaloneCachedData::FMaskedAttributeWeight> AssetAttributeWeights = BlendMask->GetSkeletonAttributeWeights();
-
-	if (AssetBoneWeights.Num() > AccumulationBoneWeights.Num())
-	{
-		AccumulationBoneWeights.AddDefaulted(AssetBoneWeights.Num() - AccumulationBoneWeights.Num());
-	}
-
-	for (FBlendProfileStandaloneCachedData::FMaskedAttributeWeight AttributeWeight : AssetAttributeWeights)
-	{
-		if (float* CurrentWeight = AttributeWeightsMap.Find(AttributeWeight.Attribute))
-		{
-			const float DampedValue = 1.0f - AttributeWeight.Weight;
-			const float EasingAdjustedWeight = DampedValue * EasingWeight;
-			*CurrentWeight = FMath::Clamp(*CurrentWeight + EasingAdjustedWeight, 0.0f, 1.0f);
-		}
-		else
-		{
-			const float DampedValue = 1.0f - AttributeWeight.Weight;
-			const float EasingAdjustedWeight = DampedValue * EasingWeight;
-			AttributeWeightsMap.Add(AttributeWeight.Attribute, FMath::Clamp(EasingAdjustedWeight, 0.0f, 1.0f));
-		}
-	}
-
-	// Since masking will remove the influence of a pose on bones, multiple masks cannot be blended as a series of tasks. 
-	// The blend task therefore needs to be a culmination of all the weights and tables contributing to this blend provider.
-	for (int32 BoneIndex = 0; BoneIndex < AssetBoneWeights.Num(); ++BoneIndex)
-	{
-		// First pass on calculating bone weight based on the asset data and easing results.
-		// This value is the influence a particular mask will have on a bone given the asset weight and easing weight.
-
-		// The weight in the asset is representative of what the expected value should be when easing is at 1.0 and no other masks are present
-		// The expected value can be thought of as the result of subtracting a damping value from the default weight of 1.0.
-		// i.e. Asset Weight = Default Weight - Damping Value, so the Damping Value = 1.0 - Asset Weight.
-
-		// The weight and easing results determine the degree to which the values from the mask should affect the underlying pose.
-		// These values act additively from mask to mask
-		// So the ultimate result for this bone's damping value will be the sum of (Easing Amount * Damping Value) for all masks.
-
-		const float AssetWeight = AssetBoneWeights[BoneIndex];
-		const float DampingValue = 1.0f - AssetWeight;
-		const float EasingAdjustedWeight = EasingWeight * DampingValue;
-		// Accumulate easing damping values.
-		const float CurrentBoneWeight = AccumulationBoneWeights[BoneIndex];
-		AccumulationBoneWeights[BoneIndex] = FMath::Clamp(CurrentBoneWeight + EasingAdjustedWeight, 0.0, 1.0);
-	}
-
-	auto GatherCurveValues = [this, &EasingWeight](UE::Anim::FCurveElement CurveElement)
-	{
-		if (float* CurrentValue = CurveWeightsMap.Find(CurveElement.Name))
-		{
-			const float DampedValue = 1.0 - CurveElement.Value;
-			const float EasingAdjustedValue = EasingWeight * DampedValue;
-			*CurrentValue = FMath::Clamp(*CurrentValue + EasingAdjustedValue, 0.0f, 1.0f);
-		}
-		else
-		{
-			const float DampedValue = 1.0 - CurveElement.Value;
-			const float EasingAdjustedValue = EasingWeight * DampedValue;
-			CurveWeightsMap.Add(CurveElement.Name, FMath::Clamp(EasingAdjustedValue, 0.0f, 1.0f));
-		}
-	};
-
-	BlendMask->GetSkeletonCurveWeights().ForEachElement(GatherCurveValues);
-}
-
-void FMovieSceneMaskBlendProvider::ResetAccumulationData()
-{
-	AccumulationBoneWeights.Reset();
-	AttributeWeightsMap.Reset();
-	CurveWeightsMap.Reset();
-	MaskSkeleton = nullptr;
-}
-
-void FMovieSceneMaskBlendProvider::ResetTaskData()
-{
-	TaskBoneWeights.Reset();
-	TaskAttributeWeights.Reset();
-	TaskCurveWeights.Empty();
-}
-
 
 UMovieSceneAnimationMaskDecoration::UMovieSceneAnimationMaskDecoration(const FObjectInitializer& ObjInit)
 	: Super(ObjInit)
@@ -205,7 +59,14 @@ void UMovieSceneAnimationMaskDecoration::AddSection(UMovieSceneSection* InSectio
 	MarkStructureChanged();
 #endif
 
-	NewMaskSection->SetAnimMixerLayer(GetTypedOuter<UMovieSceneAnimationMixerLayer>());
+	UMovieSceneAnimationMixerLayer* Layer = GetTypedOuter<UMovieSceneAnimationMixerLayer>();
+	NewMaskSection->SetAnimMixerLayer(Layer);
+
+	// Match the layer's row index so per-row mute/disable applies correctly.
+	if (Layer)
+	{
+		NewMaskSection->SetRowIndex(Layer->GetLayerIndex());
+	}
 }
 
 void UMovieSceneAnimationMaskDecoration::RemoveSection(UMovieSceneSection& SectionToRemove)
@@ -235,4 +96,4 @@ void UMovieSceneAnimationMaskDecoration::PostEditUndo()
 
 	MarkStructureChanged();
 }
-#endif
\ No newline at end of file
+#endif
diff --git a/Engine/Plugins/Experimental/MovieSceneAnimMixer/Source/MovieSceneAnimMixer/Private/MovieSceneAnimationMixerTrack.cpp b/Engine/Plugins/Experimental/MovieSceneAnimMixer/Source/MovieSceneAnimMixer/Private/MovieSceneAnimationMixerTrack.cpp
--- a/Engine/Plugins/Experimental/MovieSceneAnimMixer/Source/MovieSceneAnimMixer/Private/MovieSceneAnimationMixerTrack.cpp
+++ b/Engine/Plugins/Experimental/MovieSceneAnimMixer/Source/MovieSceneAnimMixer/Private/MovieSceneAnimationMixerTrack.cpp
@@ -1487,6 +1487,23 @@ UMovieSceneAnimationMixerLayer* UMovieSceneAnimationMixerTrack::InsertLayer(int3
 			}
 		}
 
+		// Update decoration section row indices to match the layer so
+		// per-row mute/disable applies correctly after reordering.
+		for (UObject* Decoration : Layer->GetDecorations())
+		{
+			if (IMovieSceneSectionProviderDecoration* SectionProvider = Cast<IMovieSceneSectionProviderDecoration>(Decoration))
+			{
+				for (UMovieSceneSection* DecSection : SectionProvider->GetSections())
+				{
+					if (DecSection)
+					{
+						DecSection->Modify();
+						DecSection->SetRowIndex(LayerIndex);
+					}
+				}
+			}
+		}
+
 		// Update child track row index if this layer has one
 		if (UMovieSceneTrack* ChildTrack = Layer->GetChildTrack())
 		{
diff --git a/Engine/Plugins/Experimental/MovieSceneAnimMixer/Source/MovieSceneAnimMixer/Private/Systems/MovieSceneAnimMixerMaskSystem.cpp b/Engine/Plugins/Experimental/MovieSceneAnimMixer/Source/MovieSceneAnimMixer/Private/Systems/MovieSceneAnimMixerMaskSystem.cpp
new file mode 100644
--- /dev/null
+++ b/Engine/Plugins/Experimental/MovieSceneAnimMixer/Source/MovieSceneAnimMixer/Private/Systems/MovieSceneAnimMixerMaskSystem.cpp
@@ -0,0 +1,432 @@
+// Copyright Epic Games, Inc. All Rights Reserved.
+
+#include "Systems/MovieSceneAnimMixerMaskSystem.h"
+
+#include "AnimMixerComponentTypes.h"
+#include "MovieSceneAnimationMaskDecoration.h"
+#include "MovieSceneAnimationMixerLayer.h"
+#include "EntitySystem/BuiltInComponentTypes.h"
+#include "EntitySystem/MovieSceneEntityBuilder.h"
+#include "EntitySystem/MovieSceneEntitySystemLinker.h"
+#include "EntitySystem/MovieSceneBoundObjectInstantiator.h"
+#include "EntitySystem/MovieSceneBoundSceneComponentInstantiator.h"
+#include "EntitySystem/MovieSceneEntitySystemRunner.h"
+#include "EntitySystem/MovieSceneEntitySystemTask.h"
+#include "EvaluationVM/Tasks/BlendKeyframesPerBone.h"
+#include "Systems/WeightAndEasingEvaluatorSystem.h"
+#include "Systems/MovieSceneAnimMixerSystem.h"
+#include "UAF/BlendMask/UAFBlendMask.h"
+
+#include UE_INLINE_GENERATED_CPP_BY_NAME(MovieSceneAnimMixerMaskSystem)
+
+// -----------------------------------------------------------------------
+// FCachedMaskAssetData
+// -----------------------------------------------------------------------
+
+void FCachedMaskAssetData::CacheFromAsset(UUAFBlendMask* InBlendMask)
+{
+	if (!InBlendMask)
+	{
+		DampingValues.Reset();
+		AttributeDampingValues.Reset();
+		CurveDampingValues.Reset();
+		BlendMask = nullptr;
+		return;
+	}
+
+	BlendMask = InBlendMask;
+
+	const TArray<float>& AssetBoneWeights = InBlendMask->GetSkeletonBoneWeights();
+	DampingValues.SetNumUninitialized(AssetBoneWeights.Num());
+	for (int32 Index = 0; Index < AssetBoneWeights.Num(); ++Index)
+	{
+		DampingValues[Index] = 1.0f - AssetBoneWeights[Index];
+	}
+
+	const TArray<FBlendProfileStandaloneCachedData::FMaskedAttributeWeight> AssetAttributeWeights = InBlendMask->GetSkeletonAttributeWeights();
+	AttributeDampingValues.Reset(AssetAttributeWeights.Num());
+	for (const FBlendProfileStandaloneCachedData::FMaskedAttributeWeight& AttrWeight : AssetAttributeWeights)
+	{
+		AttributeDampingValues.Add(FBlendProfileStandaloneCachedData::FMaskedAttributeWeight(AttrWeight.Attribute, 1.0f - AttrWeight.Weight));
+	}
+
+	CurveDampingValues.Reset();
+	auto GatherCurveDamping = [this](UE::Anim::FCurveElement CurveElement)
+	{
+		CurveDampingValues.Add(CurveElement.Name, 1.0f - CurveElement.Value);
+	};
+	InBlendMask->GetSkeletonCurveWeights().ForEachElement(GatherCurveDamping);
+}
+
+bool FCachedMaskAssetData::IsStale(const UUAFBlendMask* InBlendMask) const
+{
+	return BlendMask.Get() != InBlendMask || !BlendMask.IsValid();
+}
+
+// -----------------------------------------------------------------------
+// FMaskLayerAccumulation
+// -----------------------------------------------------------------------
+
+void FMaskLayerAccumulation::ResetAccumulation()
+{
+	AccumulationBoneWeights.Reset();
+	CurveWeightsMap.Reset();
+	AttributeWeightsMap.Reset();
+	MaskSkeleton = nullptr;
+	NumContributingSections = 0;
+	SingleContributingMask = nullptr;
+	bSingleSectionFullWeight = false;
+}
+
+void FMaskLayerAccumulation::AccumulateMask(UUAFBlendMask* BlendMask, double EasingWeight)
+{
+	if (!BlendMask)
+	{
+		return;
+	}
+
+	if (!MaskSkeleton.IsValid())
+	{
+		MaskSkeleton = BlendMask->GetSkeleton();
+	}
+
+	// Get or create cached data for this mask asset
+	TObjectKey<UUAFBlendMask> AssetKey(BlendMask);
+	FCachedMaskAssetData& CachedData = CachedAssets.FindOrAdd(AssetKey);
+	if (CachedData.IsStale(BlendMask))
+	{
+		CachedData.CacheFromAsset(BlendMask);
+	}
+
+	++NumContributingSections;
+
+	// Fast path: single section at full weight — can reference asset weights directly
+	if (NumContributingSections == 1 && FMath::IsNearlyEqual(EasingWeight, 1.0, KINDA_SMALL_NUMBER))
+	{
+		bSingleSectionFullWeight = true;
+		SingleContributingMask = BlendMask;
+	}
+	else
+	{
+		bSingleSectionFullWeight = false;
+		SingleContributingMask = nullptr;
+	}
+
+	// Accumulate bone damping values weighted by easing
+	if (CachedData.DampingValues.Num() > AccumulationBoneWeights.Num())
+	{
+		AccumulationBoneWeights.AddDefaulted(CachedData.DampingValues.Num() - AccumulationBoneWeights.Num());
+	}
+
+	for (int32 BoneIndex = 0; BoneIndex < CachedData.DampingValues.Num(); ++BoneIndex)
+	{
+		const float EasingAdjustedWeight = static_cast<float>(EasingWeight) * CachedData.DampingValues[BoneIndex];
+		AccumulationBoneWeights[BoneIndex] = FMath::Clamp(AccumulationBoneWeights[BoneIndex] + EasingAdjustedWeight, 0.0f, 1.0f);
+	}
+
+	// Accumulate attribute weights
+	for (const FBlendProfileStandaloneCachedData::FMaskedAttributeWeight& AttrDamping : CachedData.AttributeDampingValues)
+	{
+		const float EasingAdjustedWeight = static_cast<float>(EasingWeight) * AttrDamping.Weight;
+		if (float* CurrentWeight = AttributeWeightsMap.Find(AttrDamping.Attribute))
+		{
+			*CurrentWeight = FMath::Clamp(*CurrentWeight + EasingAdjustedWeight, 0.0f, 1.0f);
+		}
+		else
+		{
+			AttributeWeightsMap.Add(AttrDamping.Attribute, FMath::Clamp(EasingAdjustedWeight, 0.0f, 1.0f));
+		}
+	}
+
+	// Accumulate curve weights
+	for (const TTuple<FName, float>& CurveDamping : CachedData.CurveDampingValues)
+	{
+		const float EasingAdjustedValue = static_cast<float>(EasingWeight) * CurveDamping.Value;
+		if (float* CurrentValue = CurveWeightsMap.Find(CurveDamping.Key))
+		{
+			*CurrentValue = FMath::Clamp(*CurrentValue + EasingAdjustedValue, 0.0f, 1.0f);
+		}
+		else
+		{
+			CurveWeightsMap.Add(CurveDamping.Key, FMath::Clamp(EasingAdjustedValue, 0.0f, 1.0f));
+		}
+	}
+}
+
+TSharedPtr<FAnimNextEvaluationTask> FMaskLayerAccumulation::BuildBlendTask(float LayerWeight)
+{
+	if (AccumulationBoneWeights.IsEmpty())
+	{
+		return nullptr;
+	}
+
+	// Fast path: single section at full weight and layer weight is 1.0 —
+	// reference the asset's bone weight arrays directly
+	if (bSingleSectionFullWeight && FMath::IsNearlyEqual(LayerWeight, 1.0f, KINDA_SMALL_NUMBER) && SingleContributingMask.IsValid())
+	{
+		return MakeShared<FAnimNextBlendKeyframePerBoneWithScaleTask>(
+			FAnimNextBlendKeyframePerBoneWithScaleTask::Make(SingleContributingMask.Get(), 1.0f));
+	}
+
+	// Slow path: build task from accumulated damping values.
+	TaskCurveWeights.Empty();
+	TaskAttributeWeights.Reset();
+	TaskBoneWeights.Reset();
+
+	for (int32 WeightIndex = 0; WeightIndex < AccumulationBoneWeights.Num(); ++WeightIndex)
+	{
+		TaskBoneWeights.Add(1.0f - AccumulationBoneWeights[WeightIndex]);
+	}
+
+	for (const TTuple<UE::Anim::FAttributeId, float>& CurrentAttributeWeight : AttributeWeightsMap)
+	{
+		TaskAttributeWeights.Add(FBlendProfileStandaloneCachedData::FMaskedAttributeWeight(CurrentAttributeWeight.Key, 1.0f - CurrentAttributeWeight.Value));
+	}
+
+	for (const TTuple<FName, float>& CurrentCurveWeight : CurveWeightsMap)
+	{
+		TaskCurveWeights.Add(UE::Anim::FCurveElement(CurrentCurveWeight.Key, 1.0f - CurrentCurveWeight.Value));
+	}
+
+	return MakeShared<FAnimNextBlendKeyframePerBoneWithScaleTask>(
+		FAnimNextBlendKeyframePerBoneWithScaleTask::Make(MaskSkeleton.Get(), &TaskBoneWeights,
+			&TaskCurveWeights, &TaskAttributeWeights, LayerWeight));
+}
+
+// -----------------------------------------------------------------------
+// UMovieSceneAnimMixerMaskSystem
+// -----------------------------------------------------------------------
+
+namespace UE::MovieScene
+{
+
+struct FGatherMaskData
+{
+	TMap<FMaskLayerKey, FMaskLayerAccumulation>* Accumulations;
+	UMovieSceneEntitySystemLinker* Linker;
+
+	FGatherMaskData(TMap<FMaskLayerKey, FMaskLayerAccumulation>* InAccumulations, UMovieSceneEntitySystemLinker* InLinker)
+		: Accumulations(InAccumulations)
+		, Linker(InLinker)
+	{
+	}
+
+	void PreTask()
+	{
+		// Reset all accumulations before gathering
+		for (auto& [Key, Accum] : *Accumulations)
+		{
+			Accum.ResetAccumulation();
+		}
+	}
+
+	void ForEachEntity(
+		FObjectKey BoundObjectKey,
+		TObjectPtr<UMovieSceneAnimationMixerLayer> MixerLayer,
+		const TInstancedStruct<FMovieSceneAnimMixerBlendProviderData>& BlendData,
+		const double* WeightAndEasing) const
+	{
+		if (!MixerLayer)
+		{
+			return;
+		}
+
+		const FMovieSceneMaskBlendProviderData* MaskData = BlendData.GetPtr<FMovieSceneMaskBlendProviderData>();
+		if (!MaskData || !MaskData->BlendMask.IsValid())
+		{
+			return;
+		}
+
+		const double Weight = WeightAndEasing ? *WeightAndEasing : 1.0;
+		FMaskLayerKey Key{BoundObjectKey, MixerLayer};
+		FMaskLayerAccumulation& Accum = Accumulations->FindOrAdd(Key);
+		Accum.AccumulateMask(MaskData->BlendMask.Get(), Weight);
+	}
+
+	void PostTask()
+	{
+		FAnimMixerComponentTypes* AnimMixerComponents = FAnimMixerComponentTypes::Get();
+		for (auto& [Key, Accum] : *Accumulations)
+		{
+			if (!Accum.OutputEntityID)
+			{
+				continue;
+			}
+
+			FLayerBlendTaskBuilder Builder;
+			if (Accum.NumContributingSections > 0)
+			{
+				// Pre-build the task with ScaleFactor=1.0. The builder adjusts
+				// ScaleFactor to the actual layer weight each frame 
+				TSharedPtr<FAnimNextBlendKeyframePerBoneWithScaleTask> Task =
+					StaticCastSharedPtr<FAnimNextBlendKeyframePerBoneWithScaleTask>(Accum.BuildBlendTask(1.0f));
+				if (Task)
+				{
+					Builder.BuildTask = [Task](double LayerWeight) -> TSharedPtr<FAnimNextEvaluationTask>
+					{
+						Task->ScaleFactor = static_cast<float>(LayerWeight);
+						return Task;
+					};
+				}
+			}
+
+			Linker->EntityManager.WriteComponentChecked(
+				Accum.OutputEntityID, AnimMixerComponents->LayerBlendTask, MoveTemp(Builder));
+		}
+	}
+};
+
+} // namespace UE::MovieScene
+
+UMovieSceneAnimMixerMaskSystem::UMovieSceneAnimMixerMaskSystem(const FObjectInitializer& ObjInit)
+	: Super(ObjInit)
+{
+	using namespace UE::MovieScene;
+
+	FBuiltInComponentTypes* BuiltInComponents = FBuiltInComponentTypes::Get();
+	FAnimMixerComponentTypes* AnimMixerComponents = FAnimMixerComponentTypes::Get();
+
+	RelevantComponent = AnimMixerComponents->BlendProviderData;
+	Phase = ESystemPhase::Instantiation | ESystemPhase::Scheduling;
+
+	if (HasAnyFlags(RF_ClassDefaultObject))
+	{
+		DefineImplicitPrerequisite(UMovieSceneGenericBoundObjectInstantiator::StaticClass(), GetClass());
+		DefineImplicitPrerequisite(UMovieSceneBoundSceneComponentInstantiator::StaticClass(), GetClass());
+		DefineImplicitPrerequisite(UWeightAndEasingEvaluatorSystem::StaticClass(), GetClass());
+		DefineImplicitPrerequisite(GetClass(), UMovieSceneAnimMixerSystem::StaticClass());
+		DefineComponentConsumer(GetClass(), AnimMixerComponents->BlendProviderData);
+		DefineComponentConsumer(GetClass(), AnimMixerComponents->MixerLayer);
+		DefineComponentConsumer(GetClass(), BuiltInComponents->BoundObjectKey);
+		DefineComponentProducer(GetClass(), AnimMixerComponents->LayerBlendTask);
+	}
+}
+
+void UMovieSceneAnimMixerMaskSystem::OnLink()
+{
+}
+
+void UMovieSceneAnimMixerMaskSystem::UpdateAccumulationLifetimes()
+{
+	using namespace UE::MovieScene;
+	FBuiltInComponentTypes* BuiltInComponents = FBuiltInComponentTypes::Get();
+	FAnimMixerComponentTypes* AnimMixerComponents = FAnimMixerComponentTypes::Get();
+
+	// Count active (non-NeedsUnlink) mask section entities per key.
+	// This is idempotent and safe to call multiple times per frame.
+	TSet<FMaskLayerKey> ActiveKeys;
+
+	FEntityTaskBuilder()
+	.Read(BuiltInComponents->BoundObjectKey)
+	.Read(AnimMixerComponents->MixerLayer)
+	.FilterAll({ AnimMixerComponents->BlendProviderData })
+	.FilterNone({ BuiltInComponents->Tags.NeedsUnlink })
+	.Iterate_PerEntity(&Linker->EntityManager,
+		[&ActiveKeys](FObjectKey BoundObjectKey, TObjectPtr<UMovieSceneAnimationMixerLayer> MixerLayer)
+		{
+			if (MixerLayer)
+			{
+				ActiveKeys.Add(FMaskLayerKey{BoundObjectKey, MixerLayer});
+			}
+		}
+	);
+
+	// Remove accumulations whose keys have no active entities
+	for (auto It = Accumulations.CreateIterator(); It; ++It)
+	{
+		if (!ActiveKeys.Contains(It.Key()) && It.Value().OutputEntityID)
+		{
+			Linker->EntityManager.AddComponent(It.Value().OutputEntityID, BuiltInComponents->Tags.NeedsUnlink);
+			It.RemoveCurrent();
+		}
+	}
+}
+
+bool UMovieSceneAnimMixerMaskSystem::IsRelevantImpl(UMovieSceneEntitySystemLinker* InLinker) const
+{
+	using namespace UE::MovieScene;
+	FAnimMixerComponentTypes* AnimMixerComponents = FAnimMixerComponentTypes::Get();
+	return InLinker->EntityManager.ContainsAnyComponent({AnimMixerComponents->BlendProviderData});
+}
+
+void UMovieSceneAnimMixerMaskSystem::OnRun(FSystemTaskPrerequisites& InPrerequisites, FSystemSubsequentTasks& Subsequents)
+{
+	using namespace UE::MovieScene;
+
+	TSharedRef<FMovieSceneEntitySystemRunner> Runner = Linker->GetRunner();
+	if (Runner->GetCurrentPhase() != ESystemPhase::Instantiation)
+	{
+		return;
+	}
+
+	FBuiltInComponentTypes* BuiltInComponents = FBuiltInComponentTypes::Get();
+	FAnimMixerComponentTypes* AnimMixerComponents = FAnimMixerComponentTypes::Get();
+
+	// Remove output entities for keys that no longer have active mask section entities.
+	UpdateAccumulationLifetimes();
+
+	// Create output entities for newly linked mask section entities.
+	// Collected first, then created outside the iteration to avoid locking issues.
+	struct FPendingOutputEntity
+	{
+		FMaskLayerKey Key;
+		FObjectKey BoundObjectKey;
+		TObjectPtr<UMovieSceneAnimationMixerLayer> MixerLayer;
+	};
+	TArray<FPendingOutputEntity> PendingOutputEntities;
+
+	FEntityTaskBuilder()
+	.Read(BuiltInComponents->BoundObjectKey)
+	.Read(AnimMixerComponents->MixerLayer)
+	.Read(AnimMixerComponents->BlendProviderData)
+	.FilterAll({ BuiltInComponents->Tags.NeedsLink })
+	.Iterate_PerEntity(&Linker->EntityManager,
+		[this, &PendingOutputEntities](FObjectKey BoundObjectKey, TObjectPtr<UMovieSceneAnimationMixerLayer> MixerLayer, const TInstancedStruct<FMovieSceneAnimMixerBlendProviderData>& BlendData)
+		{
+			if (!MixerLayer)
+			{
+				return;
+			}
+
+			FMaskLayerKey Key{BoundObjectKey, MixerLayer};
+			if (!Accumulations.Contains(Key))
+			{
+				PendingOutputEntities.Add({ Key, BoundObjectKey, MixerLayer });
+			}
+		}
+	);
+
+	for (const FPendingOutputEntity& Pending : PendingOutputEntities)
+	{
+		FMaskLayerAccumulation& Accum = Accumulations.FindOrAdd(Pending.Key);
+		if (!Accum.OutputEntityID)
+		{
+			Accum.OutputEntityID = FEntityBuilder()
+				.Add(BuiltInComponents->BoundObjectKey, Pending.BoundObjectKey)
+				.Add(AnimMixerComponents->MixerLayer, Pending.MixerLayer)
+				.Add(AnimMixerComponents->LayerBlendTask, FLayerBlendTaskBuilder())
+				.CreateEntity(&Linker->EntityManager);
+		}
+	}
+}
+
+void UMovieSceneAnimMixerMaskSystem::OnSchedulePersistentTasks(UE::MovieScene::IEntitySystemScheduler* TaskScheduler)
+{
+	using namespace UE::MovieScene;
+
+	FBuiltInComponentTypes* BuiltInComponents = FBuiltInComponentTypes::Get();
+	FAnimMixerComponentTypes* AnimMixerComponents = FAnimMixerComponentTypes::Get();
+
+	FEntityTaskBuilder()
+	.Read(BuiltInComponents->BoundObjectKey)
+	.Read(AnimMixerComponents->MixerLayer)
+	.Read(AnimMixerComponents->BlendProviderData)
+	.ReadOptional(BuiltInComponents->WeightAndEasingResult)
+	.FilterNone({ BuiltInComponents->Tags.NeedsLink, BuiltInComponents->Tags.Ignored })
+	.Schedule_PerEntity<FGatherMaskData>(&Linker->EntityManager, TaskScheduler, &Accumulations, Linker);
+}
+
+void UMovieSceneAnimMixerMaskSystem::OnCleanTaggedGarbage()
+{
+	UpdateAccumulationLifetimes();
+}
diff --git a/Engine/Plugins/Experimental/MovieSceneAnimMixer/Source/MovieSceneAnimMixer/Private/Systems/MovieSceneAnimMixerSystem.cpp b/Engine/Plugins/Experimental/MovieSceneAnimMixer/Source/MovieSceneAnimMixer/Private/Systems/MovieSceneAnimMixerSystem.cpp
--- a/Engine/Plugins/Experimental/MovieSceneAnimMixer/Source/MovieSceneAnimMixer/Private/Systems/MovieSceneAnimMixerSystem.cpp
+++ b/Engine/Plugins/Experimental/MovieSceneAnimMixer/Source/MovieSceneAnimMixer/Private/Systems/MovieSceneAnimMixerSystem.cpp
@@ -53,7 +53,6 @@
 #include "MovieSceneAnimTransitionSectionBase.h"
 #include "EvaluationVM/Tasks/TransitionEvaluationTask.h"
 #include "EntitySystem/MovieSceneInstanceRegistry.h"
-#include "MovieSceneAnimationMaskDecoration.h"
 #include "MovieSceneAnimationMixerLayer.h"
 #include UE_INLINE_GENERATED_CPP_BY_NAME(MovieSceneAnimMixerSystem)
 
@@ -87,6 +86,9 @@ struct FProgramBuildContext
 	const UMovieSceneTransformOriginSystem* TransformOriginSystem = nullptr;
 	UMovieSceneEntitySystemLinker* Linker = nullptr;
 	USceneComponent* Component = nullptr;
+	// Layer blend task builders, keyed by MixerLayer for this specific mixer.
+	// Built from entities with {BoundObjectKey, MixerLayer, LayerBlendTask} matching this mixer's BoundObjectKey.
+	const TMap<TWeakObjectPtr<UMovieSceneAnimationMixerLayer>, FLayerBlendTaskBuilder>* LayerBlendTaskMap = nullptr;
 };
 
 // Returns true if the entry passes the filter (i.e. should be considered for inclusion)
@@ -407,11 +409,17 @@ static TSharedPtr<UE::UAF::FEvaluationProgram> BuildProgramFromEntries(
 				DeltaTime = static_cast<float>(FrameDelta.AsDecimal() / FrameRate.AsDecimal());
 			}
 
-			double BlendWeight = Entry->PoseWeight;
+			double BlendWeight = Entry->CombinedWeight;
 			TransitionTask->Update(FromTask, ToTask, BlendWeight, DeltaTime);
 		}
 
-		Entry->PoseWeight = 1.0;
+		// Compute the interpolated section weight for the layer blend.
+		// The transition's CombinedWeight is the crossfade alpha (0→1).
+		double TransitionAlpha = Entry->CombinedWeight;
+		double FromWeight = FromEntry ? FromEntry->SectionWeight : 1.0;
+		double ToWeight = ToEntry ? ToEntry->SectionWeight : 1.0;
+		Entry->SectionWeight = FMath::Lerp(FromWeight, ToWeight, TransitionAlpha);
+		Entry->CombinedWeight = 1.0;
 	}
 
 	// Third pass: collect all entries that will actually be evaluated.
@@ -458,10 +466,10 @@ static TSharedPtr<UE::UAF::FEvaluationProgram> BuildProgramFromEntries(
 
 	TSharedPtr<UE::UAF::FEvaluationProgram> Program = MakeShared<UE::UAF::FEvaluationProgram>();
 
-	// Lookahead: compute total absolute weight for a priority group
 	const int32 NumEntries = FinalEntries.Num();
 
-	auto LookaheadComputeAccumulatedWeight = [&FinalEntries, NumEntries](int32 StartAtIndex, int32 Priority, float& OutAccumulatedWeight) -> bool
+	// Helper: determine if a priority group needs a weighted average (multiple non-additive entries).
+	auto LookaheadComputeAccumulatedWeight = [&FinalEntries, NumEntries](int32 StartAtIndex, int32 Priority, bool bBackwardsCompat, float& OutAccumulatedWeight) -> bool
 	{
 		int32 NumAbsoluteBlends = 0;
 		int32 NumSkippedBlends = 0;
@@ -477,17 +485,15 @@ static TSharedPtr<UE::UAF::FEvaluationProgram> BuildProgramFromEntries(
 			{
 				break;
 			}
-			else if (NextEntry->bRequiresBlend && !NextEntry->bAdditive)
+			const double EntryWeight = bBackwardsCompat ? NextEntry->CombinedWeight : NextEntry->EasingWeight;
+			if (!FMath::IsNearlyEqual(AccumulatedWeight, AccumulatedWeight + EntryWeight, KINDA_SMALL_NUMBER))
 			{
-				if (!FMath::IsNearlyEqual(AccumulatedWeight, AccumulatedWeight + NextEntry->PoseWeight, KINDA_SMALL_NUMBER))
-				{
-					++NumAbsoluteBlends;
-					AccumulatedWeight += NextEntry->PoseWeight;
-				}
-				else
-				{
-					++NumSkippedBlends;
-				}
+				++NumAbsoluteBlends;
+				AccumulatedWeight += EntryWeight;
+			}
+			else
+			{
+				++NumSkippedBlends;
 			}
 		}
 
@@ -499,111 +505,182 @@ static TSharedPtr<UE::UAF::FEvaluationProgram> BuildProgramFromEntries(
 		return false;
 	};
 
-	// Main evaluation loop - iterate entries grouped by priority
-	for (int32 Index = 0; Index < NumEntries; )
+	// Helper: append eval + mirror + extraction + space conversion for an entry.
+	auto AppendEntryEvaluation = [&](const TSharedPtr<FMovieSceneAnimMixerEntry>& MixerEntry)
 	{
-		const TSharedPtr<FMovieSceneAnimMixerEntry>& PeekEntry = FinalEntries[Index];
-		const int32 Priority = PeekEntry->Priority;
+		Program->AppendTaskPtr(MixerEntry->EvalTask);
 
-		float AccumulatedPriorityWeight = 1.f;
-		const bool bNeedsSeparateWeightStack = LookaheadComputeAccumulatedWeight(Index, Priority, AccumulatedPriorityWeight);
+		if (MixerEntry->MirrorTable)
+		{
+			TSharedPtr<FAnimNextEvaluationMirroringTask> MirrorTask = MakeShared<FAnimNextEvaluationMirroringTask>();
+			MirrorTask->Setup.MirrorDataTable = MixerEntry->MirrorTable;
+			Program->AppendTaskPtr(MirrorTask);
+		}
 
-		bool bIsFirstAbsoluteBlend = true;
-		for ( ; Index < NumEntries; ++Index)
+		if (MixerEntry->RootMotionSettings.IsSet() && !MixerEntry->bIsTransition)
 		{
-			const TSharedPtr<FMovieSceneAnimMixerEntry>& MixerEntry = FinalEntries[Index];
-			if (MixerEntry->Priority != Priority)
-			{
-				break;
-			}
+			Program->AppendTask(FAnimNextExtractRootMotionTask::Make(
+				MixerEntry->RootMotionSettings->RootMotionSource));
+		}
 
-			if (!MixerEntry->EvalTask.IsValid())
+		if (!Filter.bSkipRootMotionConversion && !MixerEntry->bIsTransition)
+		{
+			TOptional<FAnimNextConvertRootMotionToWorldSpaceTask> SpaceTask =
+				BuildSpaceConversionTask(*MixerEntry, Context, bRootComponentHasKeyedTransform);
+			if (SpaceTask.IsSet())
 			{
-				continue;
+				if (Filter.CaptureAnimSpaceRoot)
+				{
+					SpaceTask->CaptureAnimSpaceRoot = Filter.CaptureAnimSpaceRoot;
+				}
+				Program->AppendTask(MoveTemp(SpaceTask.GetValue()));
 			}
+		}
+	};
 
-			if (bNeedsSeparateWeightStack && FMath::IsNearlyZero(MixerEntry->PoseWeight, KINDA_SMALL_NUMBER))
+	// Main evaluation loop- two-level: intra-layer, then inter-layer blend.
+	for (int32 Index = 0; Index < NumEntries; )
+	{
+		const TSharedPtr<FMovieSceneAnimMixerEntry>& PeekEntry = FinalEntries[Index];
+		const int32 Priority = PeekEntry->Priority;
+		const TWeakObjectPtr<UMovieSceneAnimationMixerLayer> Layer = PeekEntry->MixerLayer;
+		
+		// If users are running animation tracks and control rig with the mixer plugin turned on but without using a mixer track,
+		// there will be no 'layer' but we use hardcoded priorities to give the previous behavior. This requires slightly
+		// different handling of weights.
+		const bool bBackwardsCompat = !Layer.IsValid();
+
+		// ---- Level 1: Intra-layer evaluation ----
+		// Produce a single pose on the stack for this priority group.
+		double LayerWeight = 1.0;
+
+		// Check if this group has a transition entry
+		TSharedPtr<FMovieSceneAnimMixerEntry> TransitionEntry;
+		for (int32 ScanIndex = Index; ScanIndex < NumEntries; ++ScanIndex)
+		{
+			const TSharedPtr<FMovieSceneAnimMixerEntry>& ScanEntry = FinalEntries[ScanIndex];
+			if (ScanEntry->Priority != Priority) break;
+			if (ScanEntry->bIsTransition && ScanEntry->EvalTask.IsValid())
 			{
-				continue;
+				TransitionEntry = ScanEntry;
+				break;
 			}
+		}
 
-			// Evaluate the pose
-			Program->AppendTaskPtr(MixerEntry->EvalTask);
+		if (TransitionEntry)
+		{
+			// Transition: produces a single blended pose from from/to
+			AppendEntryEvaluation(TransitionEntry);
+			LayerWeight = TransitionEntry->SectionWeight;
 
-			// Mirror if specified
-			if (MixerEntry->MirrorTable)
+			// Process additives at this priority, skip everything else
+			for ( ; Index < NumEntries; ++Index)
 			{
-				TSharedPtr<FAnimNextEvaluationMirroringTask> MirrorTask = MakeShared<FAnimNextEvaluationMirroringTask>();
-				MirrorTask->Setup.MirrorDataTable = MixerEntry->MirrorTable;
-				Program->AppendTaskPtr(MirrorTask);
+				const TSharedPtr<FMovieSceneAnimMixerEntry>& MixerEntry = FinalEntries[Index];
+				if (MixerEntry->Priority != Priority)
+				{
+					break;
+				}
+				if (MixerEntry->bAdditive && MixerEntry->EvalTask.IsValid())
+				{
+					AppendEntryEvaluation(MixerEntry);
+					Program->AppendTask(FAnimNextApplyAdditiveKeyframeTask::Make(MixerEntry->EasingWeight));
+				}
 			}
+		}
+		else
+		{
+			// Count non-additive entries for weighted average decision
+			float AccumulatedGroupWeight = 1.f;
+			const bool bNeedsSeparateWeightStack = LookaheadComputeAccumulatedWeight(Index, Priority, bBackwardsCompat, AccumulatedGroupWeight);
 
-			// Extract root motion from the pose into a RootTransformAttribute.
-			// This runs after eval and mirror, before space conversion.
-			// Skipped for transitions (already handled in mini-programs).
-			if (MixerEntry->RootMotionSettings.IsSet() && !MixerEntry->bIsTransition)
-			{
-				Program->AppendTask(FAnimNextExtractRootMotionTask::Make(
-					MixerEntry->RootMotionSettings->RootMotionSource));
-			}
+			double MaxSectionWeight = 0.0;
+			bool bIsFirstAbsoluteBlend = true;
 
-			// Handle root space conversions. Skipped for bake (animation-space
-			// results) and for transitions (from/to are already converted to
-			// world space inside the wrapped mini-programs).
-			if (!Filter.bSkipRootMotionConversion && !MixerEntry->bIsTransition)
+			for ( ; Index < NumEntries; ++Index)
 			{
-				TOptional<FAnimNextConvertRootMotionToWorldSpaceTask> SpaceTask =
-					BuildSpaceConversionTask(*MixerEntry, Context, bRootComponentHasKeyedTransform);
-				if (SpaceTask.IsSet())
+				const TSharedPtr<FMovieSceneAnimMixerEntry>& MixerEntry = FinalEntries[Index];
+				if (MixerEntry->Priority != Priority)
 				{
-					if (Filter.CaptureAnimSpaceRoot)
-					{
-						SpaceTask->CaptureAnimSpaceRoot = Filter.CaptureAnimSpaceRoot;
-					}
-					Program->AppendTask(MoveTemp(SpaceTask.GetValue()));
+					break;
+				}
+
+				if (!MixerEntry->EvalTask.IsValid())
+				{
+					continue;
 				}
-			}
 
-			if (MixerEntry->bRequiresBlend)
-			{
 				if (MixerEntry->bAdditive)
 				{
-					Program->AppendTask(FAnimNextApplyAdditiveKeyframeTask::Make(MixerEntry->PoseWeight));
+					// Additives are applied after the absolute result
+					AppendEntryEvaluation(MixerEntry);
+					Program->AppendTask(FAnimNextApplyAdditiveKeyframeTask::Make(MixerEntry->EasingWeight));
+					continue;
+				}
+
+				const double IntraWeight = bBackwardsCompat ? MixerEntry->CombinedWeight : MixerEntry->EasingWeight;
+
+				if (bNeedsSeparateWeightStack && FMath::IsNearlyZero(IntraWeight, KINDA_SMALL_NUMBER))
+				{
+					continue;
 				}
-				else if (bNeedsSeparateWeightStack)
+
+				MaxSectionWeight = FMath::Max(MaxSectionWeight, MixerEntry->SectionWeight);
+
+				AppendEntryEvaluation(MixerEntry);
+
+				// Intra-layer blending (weighted average among overlapping sections)
+				if (bNeedsSeparateWeightStack)
 				{
 					if (bIsFirstAbsoluteBlend)
 					{
 						bIsFirstAbsoluteBlend = false;
-						Program->AppendTask(FAnimNextBlendOverwriteKeyframeWithScaleTask::Make(MixerEntry->PoseWeight / AccumulatedPriorityWeight));
+						Program->AppendTask(FAnimNextBlendOverwriteKeyframeWithScaleTask::Make(
+							static_cast<float>(IntraWeight / AccumulatedGroupWeight)));
 					}
 					else
 					{
-						Program->AppendTask(FMovieSceneAccumulateAbsoluteBlendTask::Make(MixerEntry->PoseWeight / AccumulatedPriorityWeight));
+						Program->AppendTask(FMovieSceneAccumulateAbsoluteBlendTask::Make(
+							static_cast<float>(IntraWeight / AccumulatedGroupWeight)));
 					}
 				}
-				else
+			}
+
+			if (bNeedsSeparateWeightStack)
+			{
+				Program->AppendTask(FAnimNextNormalizeKeyframeRotationsTask());
+				LayerWeight = bBackwardsCompat ? 1.0 : MaxSectionWeight;
+			}
+			else
+			{
+				// Single entry: combined weight (easing × section weight) controls
+				// the blend with the base/lower layer
+				LayerWeight = PeekEntry->CombinedWeight;
+			}
+		}
+
+		// ---- Level 2: Inter-layer blend ----
+		// Blend this layer's result with what's below on the stack.
+		TSharedPtr<FAnimNextEvaluationTask> LayerBlendTask;
+		if (Layer.IsValid() && Context.LayerBlendTaskMap)
+		{
+			if (const FLayerBlendTaskBuilder* Builder = Context.LayerBlendTaskMap->Find(Layer))
+			{
+				if (Builder->IsValid())
 				{
-					if (MixerEntry->BlendTask.IsValid())
-					{
-						Program->AppendTaskPtr(MixerEntry->BlendTask.Pin());
-						MixerEntry->BlendTask = nullptr;
-						// Per-bone mask blends can leave rotations un-normalized.
-						Program->AppendTask(FAnimNextNormalizeKeyframeRotationsTask());
-					}
-					else
-					{
-						Program->AppendTask(FAnimNextBlendTwoKeyframesPreserveRootMotionTask::Make(MixerEntry->PoseWeight));
-					}
+					LayerBlendTask = Builder->Build(LayerWeight);
 				}
 			}
 		}
 
-		if (bNeedsSeparateWeightStack)
+		if (LayerBlendTask)
 		{
+			Program->AppendTaskPtr(LayerBlendTask);
 			Program->AppendTask(FAnimNextNormalizeKeyframeRotationsTask());
-			constexpr float WeightedAverageBlendWeight = 1.f;
-			Program->AppendTask(FAnimNextBlendTwoKeyframesPreserveRootMotionTask::Make(WeightedAverageBlendWeight));
+		}
+		else
+		{
+			Program->AppendTask(FAnimNextBlendTwoKeyframesPreserveRootMotionTask::Make(LayerWeight));
 		}
 	}
 
@@ -622,6 +699,8 @@ namespace UE::MovieScene
 			TSharedPtr<FMovieSceneAnimMixerEntry> MixerEntry,
 			const FMovieSceneRootMotionSettings* RootMotionSettings,
 			const double* WeightAndEasings,
+			const double* EasingResult,
+			const double* WeightResult,
 			const FTransform* BoneMatchTransform,
 			const FFrameTime* OptEvalTime)
 		{
@@ -629,8 +708,19 @@ namespace UE::MovieScene
 			{
 				MixerEntry->EvalTask = Task;
 
-				// If bSkipPoseWeight is set, use 1.0 for weight
-				MixerEntry->PoseWeight = (MixerEntry->bSkipPoseWeight || !WeightAndEasings) ? 1.0 : *WeightAndEasings;
+				if (MixerEntry->bSkipPoseWeight)
+				{
+					MixerEntry->CombinedWeight = 1.0;
+					MixerEntry->EasingWeight = 1.0;
+					MixerEntry->SectionWeight = 1.0;
+				}
+				else
+				{
+					const double Combined = WeightAndEasings ? *WeightAndEasings : 1.0;
+					MixerEntry->CombinedWeight = Combined;
+					MixerEntry->EasingWeight = EasingResult ? *EasingResult : Combined;
+					MixerEntry->SectionWeight = WeightResult ? *WeightResult : 1.0;
+				}
 
 				if (RootMotionSettings)
 				{
@@ -668,64 +758,6 @@ namespace UE::MovieScene
 		}
 	};
 
-	struct FGatherBlendTasks
-	{
-		TMap<TWeakObjectPtr<UMovieSceneAnimationMixerLayer>, TWeakPtr<FMovieSceneAnimMixerEntry>>* EntriesByLayer;
-
-		FGatherBlendTasks(TMap<TWeakObjectPtr<UMovieSceneAnimationMixerLayer>, TWeakPtr<FMovieSceneAnimMixerEntry>>* InEntriesByLayer)
-			: EntriesByLayer(InEntriesByLayer)
-		{
-		}
-		~FGatherBlendTasks()
-		{
-		}
-		void ForEachEntity(TObjectPtr<UMovieSceneAnimationMixerLayer> MixerLayer, const TInstancedStruct<FMovieSceneAnimMixerBlendProviderData>& BlendData, const double* WeightAndEasing) const
-		{
-			const double Weight = WeightAndEasing ? *WeightAndEasing : 1.0f;
-			for (TObjectPtr<UObject> DecorationPtr : MixerLayer.Get()->GetDecorations())
-			{
-				if (IAnimMixerBlendProviderDecoration* BlendProviderDecoration = Cast<IAnimMixerBlendProviderDecoration>(DecorationPtr))
-				{
-					if (FMovieSceneAnimMixerBlendProvider* BlendProvider = BlendProviderDecoration->GetBlendProvider())
-					{
-						BlendProvider->RegisterTaskData(BlendData, Weight);
-					}
-				}
-			}
-		}
-		void PostTask()
-		{
-			TArray<TWeakObjectPtr<UMovieSceneAnimationMixerLayer>> MixerLayers;
-			EntriesByLayer->GetKeys(MixerLayers);
-			for (TWeakObjectPtr<UMovieSceneAnimationMixerLayer> MixerLayerKey : MixerLayers)
-			{
-				const TWeakPtr<FMovieSceneAnimMixerEntry>* Entry = EntriesByLayer->Find(MixerLayerKey);
-
-				const UMovieSceneAnimationMixerLayer* MixerLayer = MixerLayerKey.Get();
-				if (!MixerLayer)
-				{
-					continue;
-				}
-				for (TObjectPtr<UObject> DecorationPtr : MixerLayer->GetDecorations())
-				{
-					if (IAnimMixerBlendProviderDecoration* BlendProviderDecoration = Cast<IAnimMixerBlendProviderDecoration>(DecorationPtr))
-					{
-						if (FMovieSceneAnimMixerBlendProvider* BlendProvider = BlendProviderDecoration->GetBlendProvider())
-						{
-							if (Entry && Entry->IsValid() && BlendProvider)
-							{
-								Entry->Pin()->BlendTask = BlendProvider->GetBlendTask();
-								// Only one blend task per provider. The first blend provider is the one to be used.
-								break;
-							}
-						}
-					}
-				}
-			}
-			MixerLayers.Reset();
-		}
-	};
-
 	struct FEvaluateAnimMixers
 	{
 		TMap<FMovieSceneAnimMixerKey, TSharedPtr<FMovieSceneAnimMixer>>* Mixers;
@@ -735,6 +767,9 @@ namespace UE::MovieScene
 		UMovieSceneEntitySystemLinker* Linker;
 		UMovieSceneAnimMixerSystem* MixerSystem;
 
+		// Layer blend task builders grouped by BoundObjectKey → {MixerLayer → Builder}
+		TMap<FObjectKey, TMap<TWeakObjectPtr<UMovieSceneAnimationMixerLayer>, FLayerBlendTaskBuilder>> LayerBlendBuilders;
+
 		FEvaluateAnimMixers(TMap<FMovieSceneAnimMixerKey, TSharedPtr<FMovieSceneAnimMixer>>* InMixers, TMap<FObjectKey, TWeakPtr<FMovieSceneMixerRootMotionComponentData>>* InRootMotions, UMovieSceneEntitySystemLinker* InLinker, UMovieSceneAnimMixerSystem* InMixerSystem)
 			: Mixers(InMixers)
 			, RootMotions(InRootMotions)
@@ -745,6 +780,28 @@ namespace UE::MovieScene
 			TransformOriginSystem = InLinker->FindSystem<UMovieSceneTransformOriginSystem>();
 		}
 
+		void PreTask()
+		{
+			// Build layer blend task lookup from entities produced by the mask system
+			// (or other blend providers).
+			FBuiltInComponentTypes* BuiltInComponents = FBuiltInComponentTypes::Get();
+			FAnimMixerComponentTypes* AnimMixerComponents = FAnimMixerComponentTypes::Get();
+			FEntityTaskBuilder()
+			.Read(BuiltInComponents->BoundObjectKey)
+			.Read(AnimMixerComponents->MixerLayer)
+			.Read(AnimMixerComponents->LayerBlendTask)
+			.FilterNone({ BuiltInComponents->Tags.NeedsUnlink })
+			.Iterate_PerEntity(&Linker->EntityManager,
+				[this](FObjectKey BoundObjectKey, TObjectPtr<UMovieSceneAnimationMixerLayer> MixerLayer, const FLayerBlendTaskBuilder& Builder)
+				{
+					if (MixerLayer && Builder.IsValid())
+					{
+						LayerBlendBuilders.FindOrAdd(BoundObjectKey).Add(MixerLayer, Builder);
+					}
+				}
+			);
+		}
+
 		// Get an InstanceHandle from the mixer's entries (any valid entry will do).
 		// Mixer entities themselves don't carry InstanceHandle — only section entities do.
 		static FInstanceHandle GetInstanceHandleFromMixer(const FMovieSceneAnimMixer& Mixer)
@@ -784,6 +841,11 @@ namespace UE::MovieScene
 			BuildContext.Linker = Linker;
 			BuildContext.Component = Component;
 
+			// Look up layer blend builders for this bound object
+			const TMap<TWeakObjectPtr<UMovieSceneAnimationMixerLayer>, FLayerBlendTaskBuilder>* BlendMap =
+				LayerBlendBuilders.Find(MeshComponent.GetObject());
+			BuildContext.LayerBlendTaskMap = BlendMap;
+
 			// Apply bake filter from the bake target if present.
 			FProgramBuildFilter BuildFilter;
 			if (const FMovieSceneMixerBakeTarget* BakeTarget = Target.GetPtr<FMovieSceneMixerBakeTarget>())
@@ -1349,16 +1411,14 @@ void UMovieSceneAnimMixerSystem::OnRun(FSystemTaskPrerequisites& InPrerequisites
 
 				FMovieSceneAnimMixerEntry* Entry = InOutMixerEntry.Get();
 
-				if (MixerLayer && MixerLayer->Get())
-				{
-					EntriesByLayer.Add(*MixerLayer, InOutMixerEntry);
-				}
-
 				Entry->InstanceHandle = InstanceHandle;
 				Entry->EntityID = EntityID;
 				Entry->EvalTask = Task;
 				Entry->Priority = Priority;
-				Entry->PoseWeight = 0;
+				Entry->CombinedWeight = 0;
+				Entry->EasingWeight = 0;
+				Entry->SectionWeight = 1.0;
+				Entry->MixerLayer = (MixerLayer && MixerLayer->Get()) ? *MixerLayer : nullptr;
 				Entry->EntityOwner = EntityOwner ? *EntityOwner : FObjectKey();
 				if (RootMotionSettings)
 				{
@@ -1398,7 +1458,6 @@ void UMovieSceneAnimMixerSystem::OnRun(FSystemTaskPrerequisites& InPrerequisites
 				}
 				const FComponentMask& Type = Linker->EntityManager.GetEntityType(EntityID);
 				Entry->bAdditive = Type.Contains(BuiltInComponents->Tags.AdditiveAnimation);
-				Entry->bRequiresBlend = Type.Contains(AnimMixerComponents->Tags.RequiresBlending);
 				Entry->bSkipPoseWeight = Type.Contains(AnimMixerComponents->Tags.SkipPoseWeight);
 				Entry->MirrorTable = OptMirrorDataTable ? *OptMirrorDataTable : nullptr;
 				Entry->bIsTransition = Type.Contains(AnimMixerComponents->Tags.Transition);
@@ -1589,18 +1648,13 @@ void UMovieSceneAnimMixerSystem::OnSchedulePersistentTasks(UE::MovieScene::IEnti
 		.Write(AnimMixerComponents->MixerEntry)
 		.ReadOptional(AnimMixerComponents->RootMotionSettings)
 		.ReadOptional(BuiltInComponents->WeightAndEasingResult)
+		.ReadOptional(BuiltInComponents->EasingResult)
+		.ReadOptional(BuiltInComponents->WeightResult)
 		.ReadOptional(AnimMixerComponents->BoneMatchTransform)
 		.ReadOptional(BuiltInComponents->EvalTime)
 		.FilterNone({BuiltInComponents->Tags.NeedsUnlink, BuiltInComponents->Tags.Ignored})
 		.Schedule_PerEntity<FUpdateTaskEntities>(&Linker->EntityManager, TaskScheduler);
 
-	FTaskID BlendFinderTask = FEntityTaskBuilder()
-		.Read(AnimMixerComponents->MixerLayer)
-		.Read(AnimMixerComponents->BlendProviderData)
-		.ReadOptional(BuiltInComponents->WeightAndEasingResult)
-		.FilterNone({BuiltInComponents->Tags.NeedsLink, BuiltInComponents->Tags.Ignored})
-		.Schedule_PerEntity<FGatherBlendTasks>(&Linker->EntityManager, TaskScheduler, &EntriesByLayer);
-
 	// For each mixer, build the evaluation program task list.
 	FTaskID MixTask = FEntityTaskBuilder()
 		.Read(AnimMixerComponents->MeshComponent)
@@ -1610,7 +1664,6 @@ void UMovieSceneAnimMixerSystem::OnSchedulePersistentTasks(UE::MovieScene::IEnti
 		.Schedule_PerEntity<FEvaluateAnimMixers>(&Linker->EntityManager, TaskScheduler, &Mixers, &RootMotionByObject, Linker, this);
 
 	TaskScheduler->AddPrerequisite(UpdateTask, MixTask);
-	TaskScheduler->AddPrerequisite(BlendFinderTask, MixTask);
 }
 
 void UMovieSceneAnimMixerSystem::OnCleanTaggedGarbage()
diff --git a/Engine/Plugins/Experimental/MovieSceneAnimMixer/Source/MovieSceneAnimMixer/Private/Systems/MovieSceneMixedSkeletalAnimationSystem.cpp b/Engine/Plugins/Experimental/MovieSceneAnimMixer/Source/MovieSceneAnimMixer/Private/Systems/MovieSceneMixedSkeletalAnimationSystem.cpp
--- a/Engine/Plugins/Experimental/MovieSceneAnimMixer/Source/MovieSceneAnimMixer/Private/Systems/MovieSceneMixedSkeletalAnimationSystem.cpp
+++ b/Engine/Plugins/Experimental/MovieSceneAnimMixer/Source/MovieSceneAnimMixer/Private/Systems/MovieSceneMixedSkeletalAnimationSystem.cpp
@@ -89,8 +89,7 @@ struct FSkeletalAnimMixerMutation : IMovieSceneEntityMutation
 			AnimMixerComponents->Task,
 			AnimMixerComponents->MixerEntry,
 			AnimMixerComponents->RootMotionSettings,
-			AnimMixerComponents->MirrorTable,
-			AnimMixerComponents->Tags.RequiresBlending
+			AnimMixerComponents->MirrorTable
 		});
 	}
 
diff --git a/Engine/Plugins/Experimental/MovieSceneAnimMixer/Source/MovieSceneAnimMixer/Public/AnimMixerComponentTypes.h b/Engine/Plugins/Experimental/MovieSceneAnimMixer/Source/MovieSceneAnimMixer/Public/AnimMixerComponentTypes.h
--- a/Engine/Plugins/Experimental/MovieSceneAnimMixer/Source/MovieSceneAnimMixer/Public/AnimMixerComponentTypes.h
+++ b/Engine/Plugins/Experimental/MovieSceneAnimMixer/Source/MovieSceneAnimMixer/Public/AnimMixerComponentTypes.h
@@ -21,6 +21,17 @@ struct FMovieSceneAnimMixerEntry;
 
 class UMovieSceneAnimationMixerLayer;
 
+// Builder for layer-level blend tasks. Stored as an ECS component on entities
+// produced by blend type systems (e.g., the mask system). Called by the mixer
+// at evaluation time with the computed layer weight.
+struct FLayerBlendTaskBuilder
+{
+	TFunction<TSharedPtr<FAnimNextEvaluationTask>(double LayerWeight)> BuildTask;
+
+	bool IsValid() const { return static_cast<bool>(BuildTask); }
+	TSharedPtr<FAnimNextEvaluationTask> Build(double LayerWeight) const { return BuildTask ? BuildTask(LayerWeight) : nullptr; }
+};
+
 namespace UE::MovieScene
 {
 	struct FAnimMixerComponentTypes
@@ -54,9 +65,13 @@ namespace UE::MovieScene
 		// The object (section or track) that owns/created this entity
 		TComponentTypeID<FObjectKey> EntityOwner;
 
+		// Layer-level blend task builder. Provided by mask system or simple blend decorations.
+		// Entities with this + MixerLayer + BoundObjectKey define how a layer blends with layers below.
+		// The builder is called with the computed layer weight and returns the blend task.
+		TComponentTypeID<FLayerBlendTaskBuilder> LayerBlendTask;
+
 		struct
 		{
-			FComponentTypeID RequiresBlending;
 			FComponentTypeID Transition;
 			FComponentTypeID SkipPoseWeight;
 		} Tags;
diff --git a/Engine/Plugins/Experimental/MovieSceneAnimMixer/Source/MovieSceneAnimMixer/Public/MovieSceneAnimMixerBlendProvider.h b/Engine/Plugins/Experimental/MovieSceneAnimMixer/Source/MovieSceneAnimMixer/Public/MovieSceneAnimMixerBlendProvider.h
--- a/Engine/Plugins/Experimental/MovieSceneAnimMixer/Source/MovieSceneAnimMixer/Public/MovieSceneAnimMixerBlendProvider.h
+++ b/Engine/Plugins/Experimental/MovieSceneAnimMixer/Source/MovieSceneAnimMixer/Public/MovieSceneAnimMixerBlendProvider.h
@@ -2,51 +2,31 @@
 
 #pragma once
 
-#include "Templates/SharedPointer.h"
 #include "StructUtils/InstancedStruct.h"
 #include "UObject/Interface.h"
 
 #include "MovieSceneAnimMixerBlendProvider.generated.h"
 
-struct FAnimNextEvaluationTask;
-
 USTRUCT()
 struct FMovieSceneAnimMixerBlendProviderData
 {
 	GENERATED_BODY()
 };
 
-/**
- * Defines a Mixer animation blend.
- */
-USTRUCT()
-struct FMovieSceneAnimMixerBlendProvider
-{
-	GENERATED_BODY()
-	
-	virtual ~FMovieSceneAnimMixerBlendProvider() = default;
-	
-	TSharedPtr<FAnimNextEvaluationTask> BlendTask = nullptr;
-
-	virtual void RegisterTaskData(const TInstancedStruct<FMovieSceneAnimMixerBlendProviderData>& InData, const double Weight) {}
-
-	virtual TSharedPtr<FAnimNextEvaluationTask> GetBlendTask() { return BlendTask; }
-};
-
 UINTERFACE(MinimalAPI)
-class UAnimMixerBlendProviderDecoration : public UInterface
+class ULayerBlendDecoration : public UInterface
 {
 public:
 	GENERATED_BODY()
 };
 
 /**
- * Optional interface for movie scene decorations that provide a blend for animation mixer layers.
+ * Marker interface for layer blend decorations. Only one layer blend decoration
+ * is allowed per mixer layer (mutually exclusive). Used to enforce this constraint
+ * in the UI and at runtime.
  */
-class IAnimMixerBlendProviderDecoration
+class ILayerBlendDecoration
 {
 public:
 	GENERATED_BODY()
-
-	virtual FMovieSceneAnimMixerBlendProvider* GetBlendProvider() { return nullptr; }
-};
\ No newline at end of file
+};
diff --git a/Engine/Plugins/Experimental/MovieSceneAnimMixer/Source/MovieSceneAnimMixer/Public/MovieSceneAnimationMaskDecoration.h b/Engine/Plugins/Experimental/MovieSceneAnimMixer/Source/MovieSceneAnimMixer/Public/MovieSceneAnimationMaskDecoration.h
--- a/Engine/Plugins/Experimental/MovieSceneAnimMixer/Source/MovieSceneAnimMixer/Public/MovieSceneAnimationMaskDecoration.h
+++ b/Engine/Plugins/Experimental/MovieSceneAnimMixer/Source/MovieSceneAnimMixer/Public/MovieSceneAnimationMaskDecoration.h
@@ -2,11 +2,9 @@
 
 #pragma once
 
-#include "HierarchyTableBlendProfile.h"
 #include "MovieSceneAnimMixerBlendProvider.h"
 #include "MovieSceneSignedObject.h"
 #include "Decorations/IMovieSceneSectionProviderDecoration.h"
-#include "Systems/MovieSceneAnimMixerSystem.h"
 
 #include "MovieSceneAnimationMaskDecoration.generated.h"
 
@@ -21,55 +19,18 @@ struct FMovieSceneMaskBlendProviderData : public FMovieSceneAnimMixerBlendProvid
 	TWeakObjectPtr<UUAFBlendMask> BlendMask;
 };
 
-USTRUCT()
-struct FMovieSceneMaskBlendProvider : public FMovieSceneAnimMixerBlendProvider
-{
-	GENERATED_BODY()
-
-	FMovieSceneMaskBlendProvider() {}
-
-	virtual ~FMovieSceneMaskBlendProvider() = default;
-
-	virtual TSharedPtr<FAnimNextEvaluationTask> GetBlendTask() override;
-
-	virtual void RegisterTaskData(const TInstancedStruct<FMovieSceneAnimMixerBlendProviderData>& InData, const double Weight) override;
-
-	void RegisterMaskAndWeight(TObjectPtr<UUAFBlendMask> BlendMask, const double EasingWeight);
-
-
-private:
-
-	void ResetAccumulationData();
-	void ResetTaskData();
-
-	// Skeleton Associated with the blend task.
-	TWeakObjectPtr<USkeleton> MaskSkeleton;
-
-	// Per bone weights based on values from blend mask, and the section's weight and easing, additively accumulated from all tables.
-	TArray<float> AccumulationBoneWeights;
-	// Bone weights to be passed to the UAF task.
-	TArray<float> TaskBoneWeights;
-
-	// Curve weights, similar to AccumulationBoneWeights, will be additively accumulated from sections based on their weight and easing.
-	TMap<FName, float> CurveWeightsMap;
-	// Named Value Array passed to the UAF Task
-	UE::Anim::TNamedValueArray<FDefaultAllocator, UE::Anim::FCurveElement> TaskCurveWeights;
-
-	// Attribute weights, also additively accumulated based on section weight and easing.
-	TMap<UE::Anim::FAttributeId, float> AttributeWeightsMap;
-	// Array passed to the UAF Task
-	TArray<FBlendProfileStandaloneCachedData::FMaskedAttributeWeight> TaskAttributeWeights;
-};
-
 /**
- * Decorator for associating a mixer target to a per-bone blend mask (UAF Blend Mask). The anim mixer will then apply the animation only to the bones specified.
- * Bone, attribute and curve wieghts can be blended using sequencer easing.
+ * Decorator for associating a mixer layer with per-bone blend masks (UAF Blend Mask).
+ * The mask system will accumulate bone weights from all mask sections on the layer
+ * and produce a per-bone blend task for inter-layer blending.
+ *
+ * Bone, attribute and curve weights can be blended using sequencer easing.
  */
 UCLASS(MinimalAPI, DisplayName="Mask")
 class UMovieSceneAnimationMaskDecoration
 	: public UMovieSceneSignedObject
 	, public IMovieSceneSectionProviderDecoration
-	, public IAnimMixerBlendProviderDecoration
+	, public ILayerBlendDecoration
 {
 public:
 	GENERATED_BODY()
@@ -86,18 +47,10 @@ public:
 
 	MOVIESCENEANIMMIXER_API virtual bool SupportsMultipleSections() override { return true; }
 
-	MOVIESCENEANIMMIXER_API virtual FMovieSceneAnimMixerBlendProvider* GetBlendProvider() override { return &BlendProvider; }
-
 	UPROPERTY()
 	TArray<TObjectPtr<UMovieSceneSection>> MaskSections;
 
 #if WITH_EDITOR
 	virtual void PostEditUndo() override;
 #endif
-
-
-private:
-	UPROPERTY()
-	FMovieSceneMaskBlendProvider BlendProvider;
 };
-
diff --git a/Engine/Plugins/Experimental/MovieSceneAnimMixer/Source/MovieSceneAnimMixer/Public/Systems/MovieSceneAnimMixerMaskSystem.h b/Engine/Plugins/Experimental/MovieSceneAnimMixer/Source/MovieSceneAnimMixer/Public/Systems/MovieSceneAnimMixerMaskSystem.h
new file mode 100644
--- /dev/null
+++ b/Engine/Plugins/Experimental/MovieSceneAnimMixer/Source/MovieSceneAnimMixer/Public/Systems/MovieSceneAnimMixerMaskSystem.h
@@ -0,0 +1,99 @@
+// Copyright Epic Games, Inc. All Rights Reserved.
+
+#pragma once
+
+#include "EntitySystem/MovieSceneEntitySystem.h"
+#include "MovieSceneAnimMixerBlendProvider.h"
+#include "HierarchyTableBlendProfile.h"
+#include "Animation/AnimCurveTypes.h"
+#include "MovieSceneAnimMixerMaskSystem.generated.h"
+
+class UMovieSceneAnimationMixerLayer;
+class UUAFBlendMask;
+class USkeleton;
+struct FAnimNextEvaluationTask;
+
+// Cached data from a single mask asset. Invalidated when the asset changes.
+struct FCachedMaskAssetData
+{
+	TWeakObjectPtr<UUAFBlendMask> BlendMask;
+	TArray<float> DampingValues;
+	TArray<FBlendProfileStandaloneCachedData::FMaskedAttributeWeight> AttributeDampingValues;
+	// Curve damping cached as name→damping pairs
+	TMap<FName, float> CurveDampingValues;
+
+	void CacheFromAsset(UUAFBlendMask* InBlendMask);
+	bool IsStale(const UUAFBlendMask* InBlendMask) const;
+};
+
+// Per-{BoundObject, MixerLayer} accumulation state for mask blending.
+struct FMaskLayerAccumulation
+{
+	UE::MovieScene::FMovieSceneEntityID OutputEntityID;
+	TWeakObjectPtr<USkeleton> MaskSkeleton;
+
+	// Damping values accumulated from all contributing mask sections, weighted by easing
+	TArray<float> AccumulationBoneWeights;
+	TMap<FName, float> CurveWeightsMap;
+	TMap<UE::Anim::FAttributeId, float> AttributeWeightsMap;
+
+	// Final task data arrays (reused across frames to avoid allocation).
+	// These must persist past BuildBlendTask since the task holds pointers to them.
+	TArray<float> TaskBoneWeights;
+	TArray<FBlendProfileStandaloneCachedData::FMaskedAttributeWeight> TaskAttributeWeights;
+	UE::Anim::TNamedValueArray<FDefaultAllocator, UE::Anim::FCurveElement> TaskCurveWeights;
+
+	// Cached per-asset data for contributing mask sections
+	TMap<TObjectKey<UUAFBlendMask>, FCachedMaskAssetData> CachedAssets;
+
+	// Number of contributing mask sections this frame
+	int32 NumContributingSections = 0;
+	// For the single-section fast path: the mask that contributed this frame
+	TWeakObjectPtr<UUAFBlendMask> SingleContributingMask;
+	// True if the single section fast path can be used
+	bool bSingleSectionFullWeight = false;
+
+	void ResetAccumulation();
+	void AccumulateMask(UUAFBlendMask* BlendMask, double EasingWeight);
+	TSharedPtr<FAnimNextEvaluationTask> BuildBlendTask(float LayerWeight);
+};
+
+// Key for the accumulation map
+struct FMaskLayerKey
+{
+	FObjectKey BoundObjectKey;
+	TWeakObjectPtr<UMovieSceneAnimationMixerLayer> MixerLayer;
+
+	friend bool operator==(const FMaskLayerKey& A, const FMaskLayerKey& B)
+	{
+		return A.BoundObjectKey == B.BoundObjectKey && A.MixerLayer == B.MixerLayer;
+	}
+
+	friend uint32 GetTypeHash(const FMaskLayerKey& Key)
+	{
+		return HashCombineFast(GetTypeHash(Key.BoundObjectKey), GetTypeHash(Key.MixerLayer));
+	}
+};
+
+// Gathers mask section entities and produces layer blend task entities for the mixer.
+// Runs before the mixer system. Creates one output entity per {BoundObject, MixerLayer}
+// that has active mask sections, with a LayerBlendTask component.
+UCLASS(MinimalAPI)
+class UMovieSceneAnimMixerMaskSystem : public UMovieSceneEntitySystem
+{
+public:
+	GENERATED_BODY()
+
+	UMovieSceneAnimMixerMaskSystem(const FObjectInitializer& ObjInit);
+
+private:
+	TMap<FMaskLayerKey, FMaskLayerAccumulation> Accumulations;
+
+	void UpdateAccumulationLifetimes();
+
+	virtual void OnLink() override;
+	virtual bool IsRelevantImpl(UMovieSceneEntitySystemLinker* InLinker) const override;
+	virtual void OnRun(FSystemTaskPrerequisites& InPrerequisites, FSystemSubsequentTasks& Subsequents) override;
+	virtual void OnSchedulePersistentTasks(UE::MovieScene::IEntitySystemScheduler* TaskScheduler) override;
+	virtual void OnCleanTaggedGarbage() override;
+};
diff --git a/Engine/Plugins/Experimental/MovieSceneAnimMixer/Source/MovieSceneAnimMixer/Public/Systems/MovieSceneAnimMixerSystem.h b/Engine/Plugins/Experimental/MovieSceneAnimMixer/Source/MovieSceneAnimMixer/Public/Systems/MovieSceneAnimMixerSystem.h
--- a/Engine/Plugins/Experimental/MovieSceneAnimMixer/Source/MovieSceneAnimMixer/Public/Systems/MovieSceneAnimMixerSystem.h
+++ b/Engine/Plugins/Experimental/MovieSceneAnimMixer/Source/MovieSceneAnimMixer/Public/Systems/MovieSceneAnimMixerSystem.h
@@ -30,22 +30,27 @@ struct FMovieSceneAnimMixerEntry
 
 	TWeakPtr<FMovieSceneAnimMixer> WeakMixer;
 	TSharedPtr<FAnimNextEvaluationTask> EvalTask;
-	TWeakPtr<FAnimNextEvaluationTask> BlendTask;
 	/** Shared pointer to the root motion for this entry if it came from a FMovieSceneRootMotionSettings on a mixer
 	 * @note: Only to be used for lifetime management to keep FMovieSceneAnimMixer::WeakRootMotion alive!
 	 *        This may be null even if the result of the mix still has root motion */
 	TSharedPtr<FMovieSceneMixerRootMotionComponentData> RootMotionLifetimeReference;
 	TOptional<FMovieSceneRootMotionSettings> RootMotionSettings;
-	double PoseWeight = 1.0;
+	/** Combined weight and easing (WeightResult * EasingResult). Used for backwards-compat weighted average. */
+	double CombinedWeight = 1.0;
+	/** Easing-only weight from section overlap. Controls intra-layer blending when using mixer track. */
+	double EasingWeight = 1.0;
+	/** Section weight (from weight channel). Contributes to inter-layer blend weight. */
+	double SectionWeight = 1.0;
 	int32 Priority = 0;
 	bool bAdditive = false;
-	bool bRequiresBlend = true;
 	bool bSkipPoseWeight = false;
 	TObjectPtr<UMirrorDataTable> MirrorTable;
 	UE::MovieScene::FMovieSceneEntityID EntityID;
 	UE::MovieScene::FInstanceHandle InstanceHandle;
 	/** The section or track that owns this entity - used for transition section lookups */
 	FObjectKey EntityOwner;
+	/** The mixer layer this entry belongs to (null for backwards-compat standalone tracks) */
+	TWeakObjectPtr<UMovieSceneAnimationMixerLayer> MixerLayer;
 	/** True if this entry is from a transition section */
 	bool bIsTransition = false;
 
@@ -190,9 +195,6 @@ private:
 	// Map of animation mixers
 	TMap<FMovieSceneAnimMixerKey, TSharedPtr<FMovieSceneAnimMixer>> Mixers;
 
-	// Map of animation mixers by layer.
-	TMap<TWeakObjectPtr<UMovieSceneAnimationMixerLayer>, TWeakPtr<FMovieSceneAnimMixerEntry>> EntriesByLayer;
-
 	// Map of Root Motion data for each object
 	TMap<FObjectKey, TWeakPtr<FMovieSceneMixerRootMotionComponentData>> RootMotionByObject;
 
diff --git a/Engine/Plugins/Experimental/MovieSceneMixedControlRig/Source/MovieSceneMixedControlRig/Private/Systems/MovieSceneMixedControlRigSystem.cpp b/Engine/Plugins/Experimental/MovieSceneMixedControlRig/Source/MovieSceneMixedControlRig/Private/Systems/MovieSceneMixedControlRigSystem.cpp
--- a/Engine/Plugins/Experimental/MovieSceneMixedControlRig/Source/MovieSceneMixedControlRig/Private/Systems/MovieSceneMixedControlRigSystem.cpp
+++ b/Engine/Plugins/Experimental/MovieSceneMixedControlRig/Source/MovieSceneMixedControlRig/Private/Systems/MovieSceneMixedControlRigSystem.cpp
@@ -87,7 +87,6 @@ struct FControlRigMixerMutation : IMovieSceneEntityMutation
 			AnimMixerComponents->Target,
 			AnimMixerComponents->Task,
 			AnimMixerComponents->MixerEntry,
-			AnimMixerComponents->Tags.RequiresBlending,
 			BuiltInComponents->BoundObject,
 			BuiltInComponents->BoundObjectKey
 		});
diff --git a/Engine/Plugins/Experimental/MovieScenePoseSearchTracks/Source/MovieScenePoseSearchTracks/Private/Sections/MovieSceneStitchAnimSection.cpp b/Engine/Plugins/Experimental/MovieScenePoseSearchTracks/Source/MovieScenePoseSearchTracks/Private/Sections/MovieSceneStitchAnimSection.cpp
--- a/Engine/Plugins/Experimental/MovieScenePoseSearchTracks/Source/MovieScenePoseSearchTracks/Private/Sections/MovieSceneStitchAnimSection.cpp
+++ b/Engine/Plugins/Experimental/MovieScenePoseSearchTracks/Source/MovieScenePoseSearchTracks/Private/Sections/MovieSceneStitchAnimSection.cpp
@@ -398,7 +398,6 @@ void UMovieSceneStitchAnimSection::ImportEntityImpl(UMovieSceneEntitySystemLinke
 		.AddTag(TrackComponents->Tags.AnimMixerPoseProducer)
 		.Add(AnimMixerComponents->Task, EvalTask)
 		.Add(AnimMixerComponents->RootMotionSettings, RootMotionSettings)
-		.AddTag(AnimMixerComponents->Tags.RequiresBlending)
 		.Add(AnimMixerComponents->Target, TInstancedStruct<FMovieSceneMixedAnimationTarget>::Make())
 		.Add(AnimMixerComponents->MixerEntry, nullptr)
 		.Add(AnimMixerComponents->Priority, 0)
