Sequencer: Rename Root Motion modifier Keep State property and change Gap Behavior default

Rename Keep State property on Root Motion Modifier to Keep Root Motion to avoid confusion with the section When Finished Keep State concept. Change default Gap Behavior from Reset to Origin to Persist Previous Transform since most usecases prefer persisting the previous animation transform during gaps.

#Synced-CL 52466902

diff --git a/Engine/Plugins/Experimental/MovieSceneAnimMixer/Source/MovieSceneAnimMixer/Private/MovieSceneAnimationMixerTrack.cpp b/Engine/Plugins/Experimental/MovieSceneAnimMixer/Source/MovieSceneAnimMixer/Private/MovieSceneAnimationMixerTrack.cpp
--- a/Engine/Plugins/Experimental/MovieSceneAnimMixer/Source/MovieSceneAnimMixer/Private/MovieSceneAnimationMixerTrack.cpp
+++ b/Engine/Plugins/Experimental/MovieSceneAnimMixer/Source/MovieSceneAnimMixer/Private/MovieSceneAnimationMixerTrack.cpp
@@ -128,7 +128,7 @@ EMovieSceneChannelProxyType UMovieSceneRootMotionSettingsDecoration::PopulateCha
 	}
 	{
 		FMovieSceneChannelMetaData ChannelMetaData;
-		ChannelMetaData.SetIdentifiers("KeepState", LOCTEXT("KeepState", "Keep State"));
+		ChannelMetaData.SetIdentifiers("KeepState", LOCTEXT("KeepState", "Keep Root Motion"));
 		ChannelMetaData.WeakOwningObject = this;
 		ChannelMetaData.SortOrder = SortOrder++;
 		OutProxyData.Add(KeepStateChannel, ChannelMetaData, TMovieSceneExternalValue<uint8>());
diff --git a/Engine/Plugins/Experimental/MovieSceneAnimMixer/Source/MovieSceneAnimMixer/Public/MovieSceneRootMotionSection.h b/Engine/Plugins/Experimental/MovieSceneAnimMixer/Source/MovieSceneAnimMixer/Public/MovieSceneRootMotionSection.h
--- a/Engine/Plugins/Experimental/MovieSceneAnimMixer/Source/MovieSceneAnimMixer/Public/MovieSceneRootMotionSection.h
+++ b/Engine/Plugins/Experimental/MovieSceneAnimMixer/Source/MovieSceneAnimMixer/Public/MovieSceneRootMotionSection.h
@@ -109,7 +109,7 @@ public:
 
 	/** How to handle gaps where no sections with root motion are active */
 	UPROPERTY(EditAnywhere, Category="Root Motion")
-	EMovieSceneRootMotionGapBehavior GapBehavior = EMovieSceneRootMotionGapBehavior::ResetToOrigin;
+	EMovieSceneRootMotionGapBehavior GapBehavior = EMovieSceneRootMotionGapBehavior::PersistPreviousTransform;
 
 	/** Cached accumulated offsets from sections with KeepState enabled. Serialized so
 	 *  the correct values are available at runtime without needing to rebuild. */
