Sequencer: Auto-create Root Motion Target on mixer when Root Motion Settings modifier is added

When adding a Root Motion Settings modifier to an anim section (either manually or automatically for root-motion-enabled animations), automatically create the Root Motion Target modifier on the parent Animation Mixer track if one doesn't already exist.

#Synced-CL 52466902

diff --git a/Engine/Plugins/Experimental/MovieSceneAnimMixer/Source/MovieSceneAnimMixer/Private/MovieSceneAnimationMixerTrack.cpp b/Engine/Plugins/Experimental/MovieSceneAnimMixer/Source/MovieSceneAnimMixer/Private/MovieSceneAnimationMixerTrack.cpp
--- a/Engine/Plugins/Experimental/MovieSceneAnimMixer/Source/MovieSceneAnimMixer/Private/MovieSceneAnimationMixerTrack.cpp
+++ b/Engine/Plugins/Experimental/MovieSceneAnimMixer/Source/MovieSceneAnimMixer/Private/MovieSceneAnimationMixerTrack.cpp
@@ -89,6 +89,35 @@ UMovieSceneRootMotionSettingsDecoration::UMovieSceneRootMotionSettingsDecoration
 	});
 }
 
+void UMovieSceneRootMotionSettingsDecoration::OnDecorationAdded(UMovieSceneSection* Section)
+{
+	// When this Root Motion Settings decoration is added to a section that belongs
+	// to a mixer track, ensure the mixer track also has a Root Motion Target decoration.
+	if (UMovieSceneAnimationMixerTrack* MixerTrack = Section ? Section->GetTypedOuter<UMovieSceneAnimationMixerTrack>() : nullptr)
+	{
+		if (!MixerTrack->FindDecoration<UMovieSceneRootMotionTargetDecoration>())
+		{
+			MixerTrack->GetOrCreateDecoration<UMovieSceneRootMotionTargetDecoration>();
+		}
+	}
+}
+
+void UMovieSceneRootMotionSettingsDecoration::OnDecorationAdded(UMovieSceneTrack* Track)
+{
+	// When this Root Motion Settings decoration is added to a child track (e.g. Control Rig)
+	// within a mixer, ensure the mixer track also has a Root Motion Target decoration.
+	if (UMovieSceneAnimationTrackDecoration* TrackDecoration = Track ? Track->FindDecoration<UMovieSceneAnimationTrackDecoration>() : nullptr)
+	{
+		if (UMovieSceneAnimationMixerTrack* MixerTrack = TrackDecoration->MixerTrack)
+		{
+			if (!MixerTrack->FindDecoration<UMovieSceneRootMotionTargetDecoration>())
+			{
+				MixerTrack->GetOrCreateDecoration<UMovieSceneRootMotionTargetDecoration>();
+			}
+		}
+	}
+}
+
 EMovieSceneRootMotionTransformMode UMovieSceneRootMotionSettingsDecoration::GetRootTransformMode() const
 {
 	return (EMovieSceneRootMotionTransformMode)TransformMode.GetDefault().Get((uint8)EMovieSceneRootMotionTransformMode::Offset);
diff --git a/Engine/Plugins/Experimental/MovieSceneAnimMixer/Source/MovieSceneAnimMixer/Public/MovieSceneAnimationMixerTrack.h b/Engine/Plugins/Experimental/MovieSceneAnimMixer/Source/MovieSceneAnimMixer/Public/MovieSceneAnimationMixerTrack.h
--- a/Engine/Plugins/Experimental/MovieSceneAnimMixer/Source/MovieSceneAnimMixer/Public/MovieSceneAnimationMixerTrack.h
+++ b/Engine/Plugins/Experimental/MovieSceneAnimMixer/Source/MovieSceneAnimMixer/Public/MovieSceneAnimationMixerTrack.h
@@ -8,6 +8,7 @@
 #include "Decorations/IMovieSceneChannelDecoration.h"
 #include "Decorations/IMovieSceneSectionDecoration.h"
 #include "Decorations/IMovieSceneSectionProviderDecoration.h"
+#include "Decorations/IMovieSceneTrackDecoration.h"
 #include "EntitySystem/IMovieSceneEntityDecorator.h"
 #include "EntitySystem/MovieSceneEntityIDs.h"
 #include "MovieSceneMixedAnimationTarget.h"
@@ -120,6 +121,8 @@ class UMovieSceneRootMotionSettingsDecoration
 	, public IMovieSceneChannelDecoration
 	, public IMovieSceneEntityDecorator
 	, public IMovieSceneSectionProviderDecoration
+	, public IMovieSceneSectionDecoration
+	, public IMovieSceneTrackDecoration
 {
 public:
 
@@ -136,6 +139,12 @@ public:
 	// IMovieSceneSectionProviderDecoration - provides an internal section when on a track
 	virtual TArrayView<TObjectPtr<UMovieSceneSection>> GetSections() override;
 
+	// IMovieSceneSectionDecoration - auto-creates Root Motion Target when added to a section
+	virtual void OnDecorationAdded(UMovieSceneSection* Section) override;
+
+	// IMovieSceneTrackDecoration - auto-creates Root Motion Target when added to a child track
+	virtual void OnDecorationAdded(UMovieSceneTrack* Track) override;
+
 	// Creates the internal host section when the decoration is added to a track
 	MOVIESCENEANIMMIXER_API void EnsureHostSection();
 
