Sequencer: Export FAnimMixerAnimationSection and FAnimationMixerTrackEditor by request.

Move FAnimMixerAnimationSection class declaration from .cpp to .h and add MOVIESCENEANIMMIXEREDITOR_API export macro. This allows external plugins to derive from these classes to create custom ux for skel anim sections in the mixer.

#Synced-CL 52466902

diff --git a/Engine/Plugins/Experimental/MovieSceneAnimMixer/Source/MovieSceneAnimMixerEditor/Private/MovieSceneAnimationMixerTrackEditor.cpp b/Engine/Plugins/Experimental/MovieSceneAnimMixer/Source/MovieSceneAnimMixerEditor/Private/MovieSceneAnimationMixerTrackEditor.cpp
--- a/Engine/Plugins/Experimental/MovieSceneAnimMixer/Source/MovieSceneAnimMixerEditor/Private/MovieSceneAnimationMixerTrackEditor.cpp
+++ b/Engine/Plugins/Experimental/MovieSceneAnimMixer/Source/MovieSceneAnimMixerEditor/Private/MovieSceneAnimationMixerTrackEditor.cpp
@@ -202,207 +202,139 @@ private:
 	FGuid ObjectBinding;
 };
 
-class FAnimMixerAnimationSection
-	: public FCommonAnimationSection
+int32 FAnimMixerAnimationSection::OnPaintSection(FSequencerSectionPainter& Painter) const
 {
-public:
+	// Check if we're in section outliner mode (layer expanded, each section on its own row)
+	// In outliner mode, sections don't visually overlap so no transition padding is needed
+	const bool bIsInOutlinerMode = Painter.SectionModel &&
+		Painter.SectionModel->FindAncestorOfType<UE::Sequencer::FSectionOutlinerModel>().IsValid();
 
-	FAnimMixerAnimationSection(UMovieSceneSection& InSection, TWeakPtr<ISequencer> InSequencer)
-		: FCommonAnimationSection(InSection, InSequencer)
-	{
-	}
+	// Calculate transition padding using the accurate time converter before painting
+	// This value will be used by GetContentPadding() which is called after OnPaintSection()
+	CachedTransitionPaddingLeft = bIsInOutlinerMode ? 0.0f : CalculateTransitionPadding(Painter.GetTimeConverter());
 
-	virtual bool ShouldShowLegacyBoneMatching() const override { return false; }
+	return FCommonAnimationSection::OnPaintSection(Painter);
+}
 
-	virtual int32 OnPaintSection(FSequencerSectionPainter& Painter) const override
-	{
-		// Check if we're in section outliner mode (layer expanded, each section on its own row)
-		// In outliner mode, sections don't visually overlap so no transition padding is needed
-		const bool bIsInOutlinerMode = Painter.SectionModel &&
-			Painter.SectionModel->FindAncestorOfType<UE::Sequencer::FSectionOutlinerModel>().IsValid();
+FMargin FAnimMixerAnimationSection::GetContentPadding() const
+{
+	FMargin BasePadding = FCommonAnimationSection::GetContentPadding();
 
-		// Calculate transition padding using the accurate time converter before painting
-		// This value will be used by GetContentPadding() which is called after OnPaintSection()
-		CachedTransitionPaddingLeft = bIsInOutlinerMode ? 0.0f : CalculateTransitionPadding(Painter.GetTimeConverter());
+	// Add the transition padding that was calculated in OnPaintSection
+	BasePadding.Left += CachedTransitionPaddingLeft;
 
-		return FCommonAnimationSection::OnPaintSection(Painter);
-	}
+	return BasePadding;
+}
 
-	virtual FMargin GetContentPadding() const override
+float FAnimMixerAnimationSection::CalculateTransitionPadding(const FTimeToPixel& TimeConverter) const
+{
+	// Check if there's a transition section overlapping the start of this section
+	// If so, calculate extra left padding so the text appears after the transition
+	UMovieSceneSection* Section = WeakSection.Get();
+	if (!Section || !Section->HasStartFrame())
 	{
-		FMargin BasePadding = FCommonAnimationSection::GetContentPadding();
-
-		// Add the transition padding that was calculated in OnPaintSection
-		BasePadding.Left += CachedTransitionPaddingLeft;
-
-		return BasePadding;
+		return 0.0f;
 	}
 
-private:
-
-	float CalculateTransitionPadding(const FTimeToPixel& TimeConverter) const
+	UMovieSceneAnimationMixerTrack* MixerTrack = Cast<UMovieSceneAnimationMixerTrack>(Section->GetTypedOuter<UMovieSceneTrack>());
+	if (!MixerTrack)
 	{
-		// Check if there's a transition section overlapping the start of this section
-		// If so, calculate extra left padding so the text appears after the transition
-		UMovieSceneSection* Section = WeakSection.Get();
-		if (!Section || !Section->HasStartFrame())
-		{
-			return 0.0f;
-		}
-
-		UMovieSceneAnimationMixerTrack* MixerTrack = Cast<UMovieSceneAnimationMixerTrack>(Section->GetTypedOuter<UMovieSceneTrack>());
-		if (!MixerTrack)
-		{
-			return 0.0f;
-		}
+		return 0.0f;
+	}
 
-		// Find any transition that references this section as the "to" section
-		TArray<UMovieSceneAnimTransitionSectionBase*> Transitions = MixerTrack->FindTransitionsForSection(Section);
-		for (UMovieSceneAnimTransitionSectionBase* Transition : Transitions)
+	// Find any transition that references this section as the "to" section
+	TArray<UMovieSceneAnimTransitionSectionBase*> Transitions = MixerTrack->FindTransitionsForSection(Section);
+	for (UMovieSceneAnimTransitionSectionBase* Transition : Transitions)
+	{
+		if (Transition && Transition->ToSection == Section)
 		{
-			if (Transition && Transition->ToSection == Section)
+			// This is the "to" section of a transition - calculate padding for the transition width
+			TRange<FFrameNumber> TransitionRange = Transition->GetRange();
+			if (!TransitionRange.IsEmpty() && TransitionRange.HasLowerBound() && TransitionRange.HasUpperBound())
 			{
-				// This is the "to" section of a transition - calculate padding for the transition width
-				TRange<FFrameNumber> TransitionRange = Transition->GetRange();
-				if (!TransitionRange.IsEmpty() && TransitionRange.HasLowerBound() && TransitionRange.HasUpperBound())
-				{
-					// Use the accurate time converter to calculate pixel width
-					FFrameNumber TransitionStart = TransitionRange.GetLowerBoundValue();
-					FFrameNumber TransitionEnd = TransitionRange.GetUpperBoundValue();
+				// Use the accurate time converter to calculate pixel width
+				FFrameNumber TransitionStart = TransitionRange.GetLowerBoundValue();
+				FFrameNumber TransitionEnd = TransitionRange.GetUpperBoundValue();
 
-					float StartPixel = TimeConverter.FrameToPixel(TransitionStart);
-					float EndPixel = TimeConverter.FrameToPixel(TransitionEnd);
-					float TransitionPixels = EndPixel - StartPixel;
+				float StartPixel = TimeConverter.FrameToPixel(TransitionStart);
+				float EndPixel = TimeConverter.FrameToPixel(TransitionEnd);
+				float TransitionPixels = EndPixel - StartPixel;
 
-					// Add padding with a small margin
-					return TransitionPixels + 4.0f;
-				}
-				break;
+				// Add padding with a small margin
+				return TransitionPixels + 4.0f;
 			}
+			break;
 		}
-
-		return 0.0f;
 	}
 
-	// Cached transition padding calculated in OnPaintSection for use in GetContentPadding
-	mutable float CachedTransitionPaddingLeft = 0.0f;
+	return 0.0f;
+}
 
-public:
+void FAnimMixerAnimationSection::BuildSectionContextMenu(FMenuBuilder& MenuBuilder, const FGuid& ObjectBinding)
+{
+	FCommonAnimationSection::BuildSectionContextMenu(MenuBuilder, ObjectBinding);
 
-	virtual void BuildSectionContextMenu(FMenuBuilder& MenuBuilder, const FGuid& ObjectBinding) override
+	// Add "Create Transition To..." submenu for overlapping sections
+	UMovieSceneSection* Section = GetSectionObject();
+	if (Section)
 	{
-		FCommonAnimationSection::BuildSectionContextMenu(MenuBuilder, ObjectBinding);
-
-		// Add "Create Transition To..." submenu for overlapping sections
-		UMovieSceneSection* Section = GetSectionObject();
-		if (Section)
+		UMovieSceneAnimationMixerTrack* MixerTrack = Cast<UMovieSceneAnimationMixerTrack>(Section->GetTypedOuter<UMovieSceneTrack>());
+		if (MixerTrack)
 		{
-			UMovieSceneAnimationMixerTrack* MixerTrack = Cast<UMovieSceneAnimationMixerTrack>(Section->GetTypedOuter<UMovieSceneTrack>());
-			if (MixerTrack)
+			TArray<UMovieSceneSection*> OverlappingSections = MixerTrack->FindOverlappingSections(Section);
+			if (OverlappingSections.Num() > 0)
 			{
-				TArray<UMovieSceneSection*> OverlappingSections = MixerTrack->FindOverlappingSections(Section);
-				if (OverlappingSections.Num() > 0)
-				{
-					MenuBuilder.AddSubMenu(
-						LOCTEXT("CreateTransition_Label", "Create Transition To..."),
-						LOCTEXT("CreateTransition_Tooltip", "Create a transition between this section and an overlapping section"),
-						FNewMenuDelegate::CreateSP(this, &FAnimMixerAnimationSection::PopulateCreateTransitionMenu)
-					);
-				}
+				MenuBuilder.AddSubMenu(
+					LOCTEXT("CreateTransition_Label", "Create Transition To..."),
+					LOCTEXT("CreateTransition_Tooltip", "Create a transition between this section and an overlapping section"),
+					FNewMenuDelegate::CreateSP(this, &FAnimMixerAnimationSection::PopulateCreateTransitionMenu)
+				);
 			}
 		}
 	}
+}
 
-	void PopulateCreateTransitionMenu(FMenuBuilder& MenuBuilder)
+void FAnimMixerAnimationSection::PopulateCreateTransitionMenu(FMenuBuilder& MenuBuilder)
+{
+	UMovieSceneSection* Section = GetSectionObject();
+	if (!Section)
 	{
-		UMovieSceneSection* Section = GetSectionObject();
-		if (!Section)
-		{
-			return;
-		}
+		return;
+	}
 
-		UMovieSceneAnimationMixerTrack* MixerTrack = Cast<UMovieSceneAnimationMixerTrack>(Section->GetTypedOuter<UMovieSceneTrack>());
-		if (!MixerTrack)
-		{
-			return;
-		}
+	UMovieSceneAnimationMixerTrack* MixerTrack = Cast<UMovieSceneAnimationMixerTrack>(Section->GetTypedOuter<UMovieSceneTrack>());
+	if (!MixerTrack)
+	{
+		return;
+	}
 
-		TArray<UMovieSceneSection*> OverlappingSections = MixerTrack->FindOverlappingSections(Section);
+	TArray<UMovieSceneSection*> OverlappingSections = MixerTrack->FindOverlappingSections(Section);
 
-		// Get all non-abstract transition section classes
-		TArray<UClass*> TransitionClasses;
-		for (TObjectIterator<UClass> It; It; ++It)
+	// Get all non-abstract transition section classes
+	TArray<UClass*> TransitionClasses;
+	for (TObjectIterator<UClass> It; It; ++It)
+	{
+		UClass* Class = *It;
+		if (Class->IsChildOf(UMovieSceneAnimTransitionSectionBase::StaticClass()) &&
+			!Class->HasAnyClassFlags(CLASS_Abstract) &&
+			Class != UMovieSceneAnimTransitionSectionBase::StaticClass())
 		{
-			UClass* Class = *It;
-			if (Class->IsChildOf(UMovieSceneAnimTransitionSectionBase::StaticClass()) &&
-				!Class->HasAnyClassFlags(CLASS_Abstract) &&
-				Class != UMovieSceneAnimTransitionSectionBase::StaticClass())
-			{
-				TransitionClasses.Add(Class);
-			}
-		}
-
-		// Sort by display name
-		TransitionClasses.Sort([](const UClass& A, const UClass& B)
-		{
-			return A.GetDisplayNameText().CompareTo(B.GetDisplayNameText()) < 0;
-		});
-
-		for (UMovieSceneSection* OtherSection : OverlappingSections)
-		{
-			// Skip sections without start frames (unbounded sections can't have meaningful transitions)
-			if (!Section->HasStartFrame() || !OtherSection->HasStartFrame())
-			{
-				continue;
-			}
-
-			// Determine from/to based on start time
-			UMovieSceneSection* FromSection = Section;
-			UMovieSceneSection* ToSection = OtherSection;
-			if (OtherSection->GetInclusiveStartFrame() < Section->GetInclusiveStartFrame())
-			{
-				FromSection = OtherSection;
-				ToSection = Section;
-			}
-
-			// Check if a transition already exists for this pair
-			UMovieSceneAnimTransitionSectionBase* ExistingTransition = MixerTrack->FindTransitionForPair(FromSection, ToSection);
-			if (ExistingTransition)
-			{
-				continue; // Skip - transition already exists
-			}
-
-			// Get a display name for the other section
-			FText OtherSectionName = FText::FromString(OtherSection->GetName());
-			if (UMovieSceneSkeletalAnimationSection* OtherAnimSection = Cast<UMovieSceneSkeletalAnimationSection>(OtherSection))
-			{
-				if (UAnimSequenceBase* Animation = OtherAnimSection->GetAnimation())
-				{
-					OtherSectionName = FText::FromString(Animation->GetName());
-				}
-			}
-
-			// Add submenu for this overlapping section with transition type options
-			MenuBuilder.AddSubMenu(
-				OtherSectionName,
-				FText::Format(LOCTEXT("CreateTransitionTo_Tooltip", "Create a transition to {0}"), OtherSectionName),
-				FNewMenuDelegate::CreateSP(this, &FAnimMixerAnimationSection::PopulateTransitionTypeMenu, OtherSection, TransitionClasses)
-			);
+			TransitionClasses.Add(Class);
 		}
 	}
 
-	void PopulateTransitionTypeMenu(FMenuBuilder& MenuBuilder, UMovieSceneSection* OtherSection, TArray<UClass*> TransitionClasses)
+	// Sort by display name
+	TransitionClasses.Sort([](const UClass& A, const UClass& B)
 	{
-		UMovieSceneSection* Section = GetSectionObject();
-		if (!Section || !OtherSection)
-		{
-			return;
-		}
+		return A.GetDisplayNameText().CompareTo(B.GetDisplayNameText()) < 0;
+	});
 
-		// Skip sections without start frames
+	for (UMovieSceneSection* OtherSection : OverlappingSections)
+	{
+		// Skip sections without start frames (unbounded sections can't have meaningful transitions)
 		if (!Section->HasStartFrame() || !OtherSection->HasStartFrame())
 		{
-			return;
+			continue;
 		}
 
 		// Determine from/to based on start time
@@ -414,43 +346,101 @@ public:
 			ToSection = Section;
 		}
 
-		for (UClass* TransitionClass : TransitionClasses)
+		// Check if a transition already exists for this pair
+		UMovieSceneAnimTransitionSectionBase* ExistingTransition = MixerTrack->FindTransitionForPair(FromSection, ToSection);
+		if (ExistingTransition)
 		{
-			FText DisplayName = TransitionClass->GetDisplayNameText();
+			continue; // Skip - transition already exists
+		}
 
-			MenuBuilder.AddMenuEntry(
-				DisplayName,
-				FText::Format(LOCTEXT("CreateTransitionOfType_Tooltip", "Create a {0} transition"), DisplayName),
-				FSlateIcon(),
-				FUIAction(
-					FExecuteAction::CreateSP(this, &FAnimMixerAnimationSection::CreateTransitionToSection, FromSection, ToSection, TransitionClass)
-				)
-			);
+		// Get a display name for the other section
+		FText OtherSectionName = FText::FromString(OtherSection->GetName());
+		if (UMovieSceneSkeletalAnimationSection* OtherAnimSection = Cast<UMovieSceneSkeletalAnimationSection>(OtherSection))
+		{
+			if (UAnimSequenceBase* Animation = OtherAnimSection->GetAnimation())
+			{
+				OtherSectionName = FText::FromString(Animation->GetName());
+			}
 		}
+
+		// Add submenu for this overlapping section with transition type options
+		MenuBuilder.AddSubMenu(
+			OtherSectionName,
+			FText::Format(LOCTEXT("CreateTransitionTo_Tooltip", "Create a transition to {0}"), OtherSectionName),
+			FNewMenuDelegate::CreateSP(this, &FAnimMixerAnimationSection::PopulateTransitionTypeMenu, OtherSection, TransitionClasses)
+		);
 	}
+}
 
-	void CreateTransitionToSection(UMovieSceneSection* FromSection, UMovieSceneSection* ToSection, UClass* TransitionClass)
+void FAnimMixerAnimationSection::PopulateTransitionTypeMenu(FMenuBuilder& MenuBuilder, UMovieSceneSection* OtherSection, TArray<UClass*> TransitionClasses)
+{
+	UMovieSceneSection* Section = GetSectionObject();
+	if (!Section || !OtherSection)
 	{
-		if (!FromSection || !ToSection || !TransitionClass)
-		{
-			return;
-		}
+		return;
+	}
 
-		UMovieSceneAnimationMixerTrack* MixerTrack = Cast<UMovieSceneAnimationMixerTrack>(FromSection->GetTypedOuter<UMovieSceneTrack>());
-		if (!MixerTrack)
-		{
-			return;
-		}
+	// Skip sections without start frames
+	if (!Section->HasStartFrame() || !OtherSection->HasStartFrame())
+	{
+		return;
+	}
 
-		FScopedTransaction Transaction(LOCTEXT("CreateTransition", "Create Transition"));
+	// Determine from/to based on start time
+	UMovieSceneSection* FromSection = Section;
+	UMovieSceneSection* ToSection = OtherSection;
+	if (OtherSection->GetInclusiveStartFrame() < Section->GetInclusiveStartFrame())
+	{
+		FromSection = OtherSection;
+		ToSection = Section;
+	}
 
-		MixerTrack->CreateTransitionSectionOfType(FromSection, ToSection, TransitionClass);
+	for (UClass* TransitionClass : TransitionClasses)
+	{
+		FText DisplayName = TransitionClass->GetDisplayNameText();
+
+		MenuBuilder.AddMenuEntry(
+			DisplayName,
+			FText::Format(LOCTEXT("CreateTransitionOfType_Tooltip", "Create a {0} transition"), DisplayName),
+			FSlateIcon(),
+			FUIAction(
+				FExecuteAction::CreateSP(this, &FAnimMixerAnimationSection::CreateTransitionToSection, FromSection, ToSection, TransitionClass)
+			)
+		);
 	}
-};
+}
+
+void FAnimMixerAnimationSection::CreateTransitionToSection(UMovieSceneSection* FromSection, UMovieSceneSection* ToSection, UClass* TransitionClass)
+{
+	if (!FromSection || !ToSection || !TransitionClass)
+	{
+		return;
+	}
+
+	UMovieSceneAnimationMixerTrack* MixerTrack = Cast<UMovieSceneAnimationMixerTrack>(FromSection->GetTypedOuter<UMovieSceneTrack>());
+	if (!MixerTrack)
+	{
+		return;
+	}
+
+	FScopedTransaction Transaction(LOCTEXT("CreateTransition", "Create Transition"));
+
+	MixerTrack->CreateTransitionSectionOfType(FromSection, ToSection, TransitionClass);
+}
 
 int32 FAnimationMixerTrackEditor::MixerEditModeRefCount = 0;
 TMap<const UClass*, FOnMakeSectionInterfaceDelegate> FAnimationMixerTrackEditor::MakeSectionInterfaceCallbacks;
 
+void FAnimationMixerTrackEditor::RegisterMakeSectionInterfaceCallback(const UClass* SectionClass, FOnMakeSectionInterfaceDelegate Callback)
+{
+	MakeSectionInterfaceCallbacks.Add(SectionClass, MoveTemp(Callback));
+}
+
+void FAnimationMixerTrackEditor::UnregisterMakeSectionInterfaceCallback(const UClass* SectionClass)
+{
+	MakeSectionInterfaceCallbacks.Remove(SectionClass);
+}
+
 TSharedRef<ISequencerTrackEditor> FAnimationMixerTrackEditor::CreateTrackEditor(TSharedRef<ISequencer> InSequencer)
 {
 	return MakeShareable(new FAnimationMixerTrackEditor(InSequencer));
diff --git a/Engine/Plugins/Experimental/MovieSceneAnimMixer/Source/MovieSceneAnimMixerEditor/Private/MovieSceneAnimationMixerTrackEditor.h b/Engine/Plugins/Experimental/MovieSceneAnimMixer/Source/MovieSceneAnimMixerEditor/Private/MovieSceneAnimationMixerTrackEditor.h
--- a/Engine/Plugins/Experimental/MovieSceneAnimMixer/Source/MovieSceneAnimMixerEditor/Private/MovieSceneAnimationMixerTrackEditor.h
+++ b/Engine/Plugins/Experimental/MovieSceneAnimMixer/Source/MovieSceneAnimMixerEditor/Private/MovieSceneAnimationMixerTrackEditor.h
@@ -11,6 +11,8 @@
 #include "TrackEditors/CommonAnimationTrackEditor.h"
 #include "IMovieSceneAnimMixerItemMenuProvider.h"
 
+struct FTimeToPixel;
+
 namespace UE::SequencerAnimTools { class FTrailHierarchy; }
 namespace UE::Sequencer { class FAnimMixerTrailHierarchy; }
 
@@ -19,7 +21,38 @@ DECLARE_DELEGATE_RetVal_ThreeParams(TSharedRef<ISequencerSection>, FOnMakeSectio
 namespace UE::Sequencer
 {
 
-class FAnimationMixerTrackEditor
+// Section interface for skeletal animation sections within a mixer track. Handles transition
+// padding, context menus for creating transitions between overlapping sections, and
+// disables the legacy bone matching UI. 
+class MOVIESCENEANIMMIXEREDITOR_API FAnimMixerAnimationSection
+	: public FCommonAnimationSection
+{
+public:
+
+	FAnimMixerAnimationSection(UMovieSceneSection& InSection, TWeakPtr<ISequencer> InSequencer)
+		: FCommonAnimationSection(InSection, InSequencer)
+	{
+	}
+
+	virtual bool ShouldShowLegacyBoneMatching() const override { return false; }
+
+	virtual int32 OnPaintSection(FSequencerSectionPainter& Painter) const override;
+	virtual FMargin GetContentPadding() const override;
+	virtual void BuildSectionContextMenu(FMenuBuilder& MenuBuilder, const FGuid& ObjectBinding) override;
+
+	void PopulateCreateTransitionMenu(FMenuBuilder& MenuBuilder);
+	void PopulateTransitionTypeMenu(FMenuBuilder& MenuBuilder, UMovieSceneSection* OtherSection, TArray<UClass*> TransitionClasses);
+	void CreateTransitionToSection(UMovieSceneSection* FromSection, UMovieSceneSection* ToSection, UClass* TransitionClass);
+
+private:
+
+	float CalculateTransitionPadding(const FTimeToPixel& TimeConverter) const;
+
+	// Cached transition padding calculated in OnPaintSection for use in GetContentPadding
+	mutable float CachedTransitionPaddingLeft = 0.0f;
+};
+
+class MOVIESCENEANIMMIXEREDITOR_API FAnimationMixerTrackEditor
 	: public FCommonAnimationTrackEditor
 	, public IMovieSceneAnimMixerItemMenuProvider
 {
@@ -83,8 +116,12 @@ private:
 	TSharedPtr<FAnimMixerTrailHierarchy> MixerTrailHierarchy;
 
 	static int32 MixerEditModeRefCount;
-	MOVIESCENEANIMMIXEREDITOR_API static TMap<const UClass*, FOnMakeSectionInterfaceDelegate> MakeSectionInterfaceCallbacks;
+	static void RegisterMakeSectionInterfaceCallback(const UClass* SectionClass, FOnMakeSectionInterfaceDelegate Callback);
+	static void UnregisterMakeSectionInterfaceCallback(const UClass* SectionClass);
+
+private:
+	static TMap<const UClass*, FOnMakeSectionInterfaceDelegate> MakeSectionInterfaceCallbacks;
 };
 
 
-} // namespace UE::Sequencer
\ No newline at end of file
+} // namespace UE::Sequencer
