Sequencer: Fix muted Control Rig mix layers still affecting evaluation.

Muted child track layers (like Control Rig) were not filtered out in BuildProgramFromEntries because the mute check only applied to the child track's own row state, not the mixer track's. Added IsRowEvalDisabled check against the mixer track for each entry's priority/row.

#Synced-CL 52466902

diff --git a/Engine/Plugins/Experimental/MovieSceneAnimMixer/Source/MovieSceneAnimMixer/Private/Systems/MovieSceneAnimMixerSystem.cpp b/Engine/Plugins/Experimental/MovieSceneAnimMixer/Source/MovieSceneAnimMixer/Private/Systems/MovieSceneAnimMixerSystem.cpp
--- a/Engine/Plugins/Experimental/MovieSceneAnimMixer/Source/MovieSceneAnimMixer/Private/Systems/MovieSceneAnimMixerSystem.cpp
+++ b/Engine/Plugins/Experimental/MovieSceneAnimMixer/Source/MovieSceneAnimMixer/Private/Systems/MovieSceneAnimMixerSystem.cpp
@@ -292,6 +292,14 @@ static TSharedPtr<UE::UAF::FEvaluationProgram> BuildProgramFromEntries(
 		}
 		if (PassesFilter(*Entry, Filter))
 		{
+			// Skip entries whose mixer layer row is muted/disabled on the mixer track.
+			if (UMovieSceneAnimationMixerTrack* MixerTrack = Mixer.OwnerTrack.Get())
+			{
+				if (MixerTrack->IsRowEvalDisabled(Entry->Priority))
+				{
+					continue;
+				}
+			}
 			IncludedSectionOwners.Add(Entry->EntityOwner);
 		}
 	}
