Sequencer: Support looping skel anim sections for accumulated offset root motion.

When animation sections extend beyond their source length and loop, the
accumulated offset cache had no bake points at loop boundaries, causing
incorrect offset computation across loops.

Changes FRootMotionAccumulatedOffset::StartFrame from FFrameNumber to
FFrameTime so loop boundary cache entries preserve sub-frame precision.

Invalidates the accumulated offset cache when the owning section's timing
parameters change (StartFrameOffset, EndFrameOffset, FirstLoopStartFrameOffset).

#Synced-CL 52530536

diff --git a/Engine/Plugins/Experimental/MovieSceneAnimMixer/Source/MovieSceneAnimMixer/Private/MovieSceneAnimationMixerTrack.cpp b/Engine/Plugins/Experimental/MovieSceneAnimMixer/Source/MovieSceneAnimMixer/Private/MovieSceneAnimationMixerTrack.cpp
--- a/Engine/Plugins/Experimental/MovieSceneAnimMixer/Source/MovieSceneAnimMixer/Private/MovieSceneAnimationMixerTrack.cpp
+++ b/Engine/Plugins/Experimental/MovieSceneAnimMixer/Source/MovieSceneAnimMixer/Private/MovieSceneAnimationMixerTrack.cpp
@@ -80,13 +80,21 @@ UMovieSceneRootMotionSettingsDecoration::UMovieSceneRootMotionSettingsDecoration
 	// Invalidate root motion caches whenever this decoration is modified.
 	// This covers channel default edits (which only call Modify(), not
 	// PostEditChangeProperty) as well as property edits and undo.
-	OnSignatureChanged().AddWeakLambda(this, [this]()
+	auto InvalidateCacheLambda = [this]()
 	{
 		if (UMovieSceneAnimationMixerTrack* MixerTrack = GetTypedOuter<UMovieSceneAnimationMixerTrack>())
 		{
 			MixerTrack->InvalidateAccumulatedOffsetCache();
 		}
-	});
+	};
+
+	OnSignatureChanged().AddWeakLambda(this, InvalidateCacheLambda);
+
+	// Also invalidate when the owning section's properties change.
+	if (UMovieSceneSection* OwnerSection = GetTypedOuter<UMovieSceneSection>())
+	{
+		OwnerSection->OnSignatureChanged().AddWeakLambda(this, InvalidateCacheLambda);
+	}
 }
 
 void UMovieSceneRootMotionSettingsDecoration::OnDecorationAdded(UMovieSceneSection* Section)
@@ -118,6 +126,55 @@ void UMovieSceneRootMotionSettingsDecoration::OnDecorationAdded(UMovieSceneTrack
 	}
 }
 
+TOptional<FFrameTime> UMovieSceneRootMotionSettingsDecoration::GetFirstLoopBoundary(const FFrameRate& TickResolution) const
+{
+	UMovieSceneSection* OwnerSection = GetTypedOuter<UMovieSceneSection>();
+	auto* SkelAnimSection = OwnerSection ? Cast<UMovieSceneSkeletalAnimationSection>(OwnerSection) : nullptr;
+	if (!SkelAnimSection || !OwnerSection->HasStartFrame() || !OwnerSection->HasEndFrame())
+	{
+		return {};
+	}
+
+	const float SequenceLength = SkelAnimSection->Params.GetSequenceLength();
+	if (SequenceLength <= KINDA_SMALL_NUMBER)
+	{
+		return {};
+	}
+
+	// Use the section's timing transform to find the actual loop boundary.
+	// Manual frame arithmetic with FloorToFrame() loses sub-frames, which causes
+	// the boundary to land before the timing system's actual loop wrap point.
+	FMovieSceneSequenceTransform Transform = SkelAnimSection->Params.MakeTransform(
+		TickResolution, OwnerSection->GetRange());
+
+	TOptional<FFrameTime> FirstBoundary;
+	Transform.ExtractBoundariesWithinRange(
+		FFrameTime(OwnerSection->GetInclusiveStartFrame()),
+		FFrameTime(OwnerSection->GetExclusiveEndFrame()),
+		[&FirstBoundary](FFrameTime Boundary) -> bool
+		{
+			FirstBoundary = Boundary;
+			return false; // stop after first
+		});
+
+	return FirstBoundary;
+}
+
+FTransform UMovieSceneRootMotionSettingsDecoration::GetInitialRootTransformAtTime(FFrameTime Time, const FFrameRate& TickResolution) const
+{
+	if (!bInitialRootTransformValid)
+	{
+		return FTransform::Identity;
+	}
+
+	TOptional<FFrameTime> FirstLoop = GetFirstLoopBoundary(TickResolution);
+	if (FirstLoop.IsSet() && Time >= FirstLoop.GetValue())
+	{
+		return LoopStartRootTransform;
+	}
+	return InitialRootTransform;
+}
+
 EMovieSceneRootMotionTransformMode UMovieSceneRootMotionSettingsDecoration::GetRootTransformMode() const
 {
 	return (EMovieSceneRootMotionTransformMode)TransformMode.GetDefault().Get((uint8)EMovieSceneRootMotionTransformMode::Offset);
@@ -1617,6 +1674,26 @@ void UMovieSceneAnimationMixerTrack::RebuildDirtyAccumulatedOffsetCache(
 					Linker, InstanceHandle, this, SectionStart, Filter);
 
 				Settings->InitialRootTransform = Result.RootMotionTransform;
+
+				// Compute the root transform at the animation's natural start for
+				// post-first-loop iterations. When the section starts mid-animation
+				// (extended left), the first loop wraps back to anim frame 0 which
+				// has a different root than the section start.
+				{
+					const FFrameRate TickResolution = GetTypedOuter<UMovieScene>()->GetTickResolution();
+					TOptional<FFrameTime> FirstLoopBoundary = Settings->GetFirstLoopBoundary(TickResolution);
+					if (FirstLoopBoundary.IsSet())
+					{
+						FBakeResult LoopResult = EvaluateAtTime(
+							Linker, InstanceHandle, this, FirstLoopBoundary.GetValue(), Filter);
+						Settings->LoopStartRootTransform = LoopResult.RootMotionTransform;
+					}
+					else
+					{
+						Settings->LoopStartRootTransform = Result.RootMotionTransform;
+					}
+				}
+
 				Settings->bInitialRootTransformValid = true;
 			}
 
diff --git a/Engine/Plugins/Experimental/MovieSceneAnimMixer/Source/MovieSceneAnimMixer/Private/MovieSceneRootMotionSection.cpp b/Engine/Plugins/Experimental/MovieSceneAnimMixer/Source/MovieSceneAnimMixer/Private/MovieSceneRootMotionSection.cpp
--- a/Engine/Plugins/Experimental/MovieSceneAnimMixer/Source/MovieSceneAnimMixer/Private/MovieSceneRootMotionSection.cpp
+++ b/Engine/Plugins/Experimental/MovieSceneAnimMixer/Source/MovieSceneAnimMixer/Private/MovieSceneRootMotionSection.cpp
@@ -8,7 +8,9 @@
 #include "EntitySystem/BuiltInComponentTypes.h"
 #include "Evaluation/MovieSceneEvaluationField.h"
 #include "MovieSceneAnimTransitionSectionBase.h"
+#include "MovieScene.h"
 #include "MovieSceneTrack.h"
+#include "Sections/MovieSceneSkeletalAnimationSection.h"
 #include "Systems/MovieSceneSkeletalAnimationSystem.h"
 #include "Tracks/MovieSceneCommonAnimationTrack.h"
 
@@ -102,11 +104,10 @@ bool UMovieSceneRootMotionSection::PopulateEvaluationFieldImpl(const TRange<FFra
 FTransform UMovieSceneRootMotionSection::GetAccumulatedOffsetAt(FFrameTime Time) const
 {
 	FTransform Result = FTransform::Identity;
-	FFrameNumber TimeFrame = Time.FloorToFrame();
 
 	for (const FRootMotionAccumulatedOffset& Offset : AccumulatedOffsetCache)
 	{
-		if (Offset.StartFrame <= TimeFrame)
+		if (Offset.StartFrame <= Time)
 		{
 			Result = Offset.AccumulatedOffset;
 		}
@@ -122,11 +123,10 @@ FTransform UMovieSceneRootMotionSection::GetAccumulatedOffsetAt(FFrameTime Time)
 FTransform UMovieSceneRootMotionSection::GetAnimSpaceAccumulatedOffsetAt(FFrameTime Time) const
 {
 	FTransform Result = FTransform::Identity;
-	FFrameNumber TimeFrame = Time.FloorToFrame();
 
 	for (const FRootMotionAccumulatedOffset& Offset : AccumulatedOffsetCache)
 	{
-		if (Offset.StartFrame <= TimeFrame)
+		if (Offset.StartFrame <= Time)
 		{
 			Result = Offset.AnimSpaceAccumulatedOffset;
 		}
@@ -191,7 +191,7 @@ void UMovieSceneRootMotionSection::RebuildAccumulatedOffsetCache(
 	// lookup finds correct values from previously computed entries.
 	struct FPendingEntry
 	{
-		FFrameNumber EntryStartFrame;
+		FFrameTime EntryStartFrame;
 		UMovieSceneSection* Section;
 		FFrameTime EvalTime;
 	};
@@ -205,6 +205,36 @@ void UMovieSceneRootMotionSection::RebuildAccumulatedOffsetCache(
 		Entry.Section = Section;
 		Entry.EvalTime = FFrameTime(Section->GetExclusiveEndFrame() - 1);
 		PendingEntries.Add(Entry);
+
+		// For looping animation sections, add bake points at each loop boundary.
+		if (UMovieSceneSkeletalAnimationSection* SkelAnimSection = Cast<UMovieSceneSkeletalAnimationSection>(Section))
+		{
+			const float SequenceLength = SkelAnimSection->Params.GetSequenceLength();
+			if (SequenceLength > KINDA_SMALL_NUMBER && Section->HasStartFrame() && Section->HasEndFrame())
+			{
+				const UMovieScene* MovieScene = Section->GetTypedOuter<UMovieScene>();
+				if (MovieScene)
+				{
+					const FFrameRate TickResolution = MovieScene->GetTickResolution();
+
+					FMovieSceneSequenceTransform Transform = SkelAnimSection->Params.MakeTransform(
+						TickResolution, Section->GetRange());
+
+					Transform.ExtractBoundariesWithinRange(
+						FFrameTime(Section->GetInclusiveStartFrame()),
+						FFrameTime(Section->GetExclusiveEndFrame()),
+						[&PendingEntries, Section](FFrameTime Boundary) -> bool
+						{
+							FPendingEntry LoopEntry;
+							LoopEntry.EntryStartFrame = Boundary;
+							LoopEntry.Section = Section;
+							LoopEntry.EvalTime = Boundary - 1;
+							PendingEntries.Add(LoopEntry);
+							return true; // continue to get all boundaries
+						});
+				}
+			}
+		}
 	}
 
 	for (UMovieSceneAnimTransitionSectionBase* Transition : Transitions)
@@ -248,6 +278,40 @@ void UMovieSceneRootMotionSection::RebuildAccumulatedOffsetCache(
 		return A.EntryStartFrame < B.EntryStartFrame;
 	});
 
+	// For looping AccumulatedOffset sections, seed the cache with an entry at the
+	// section start whose AccumulatedOffset = InitialRootTransform.
+	for (UMovieSceneSection* Section : KeepStateSectionSet)
+	{
+		auto* SkelAnimSection = Cast<UMovieSceneSkeletalAnimationSection>(Section);
+		if (!SkelAnimSection || !Section->HasStartFrame())
+		{
+			continue;
+		}
+		auto* Settings = Section->FindDecoration<UMovieSceneRootMotionSettingsDecoration>();
+		if (!Settings
+			|| Settings->GetRootTransformMode() != EMovieSceneRootMotionTransformMode::AccumulatedOffset
+			|| !Settings->bInitialRootTransformValid)
+		{
+			continue;
+		}
+		const UMovieScene* SeedMovieScene = MixerTrack->GetTypedOuter<UMovieScene>();
+		if (!SeedMovieScene)
+		{
+			continue;
+		}
+		TOptional<FFrameTime> FirstBoundary = Settings->GetFirstLoopBoundary(
+			SeedMovieScene->GetTickResolution());
+		if (FirstBoundary.IsSet())
+		{
+			FRootMotionAccumulatedOffset InitialEntry;
+			InitialEntry.StartFrame = FFrameTime(Section->GetInclusiveStartFrame());
+			InitialEntry.AccumulatedOffset = Settings->InitialRootTransform;
+			InitialEntry.AnimSpaceAccumulatedOffset = Settings->InitialRootTransform;
+			InitialEntry.Weight = 1.0f;
+			AccumulatedOffsetCache.Add(InitialEntry);
+		}
+	}
+
 	for (const FPendingEntry& Pending : PendingEntries)
 	{
 		auto* Settings = Pending.Section->FindDecoration<UMovieSceneRootMotionSettingsDecoration>();
@@ -265,7 +329,11 @@ void UMovieSceneRootMotionSection::RebuildAccumulatedOffsetCache(
 			AnimMixerBakeEvaluation::FBakeResult Result = AnimMixerBakeEvaluation::EvaluateAtTime(
 				Linker, InstanceHandle, MixerTrack, Pending.EvalTime, Filter);
 
-			FTransform Delta = Settings->InitialRootTransform.Inverse() * Result.RootMotionTransform;
+			const UMovieScene* CacheMovieScene = MixerTrack->GetTypedOuter<UMovieScene>();
+			FTransform InitialRoot = CacheMovieScene
+				? Settings->GetInitialRootTransformAtTime(Pending.EvalTime, CacheMovieScene->GetTickResolution())
+				: Settings->InitialRootTransform;
+			FTransform Delta = InitialRoot.Inverse() * Result.RootMotionTransform;
 
 			// Apply the user-keyed offset to the delta so that subsequent
 			// sections accumulate on top of the offset. Mirrors the
@@ -297,9 +365,11 @@ void UMovieSceneRootMotionSection::RebuildAccumulatedOffsetCache(
 
 			if (Pending.Section->HasStartFrame())
 			{
-				FFrameTime StartTime(Pending.Section->GetInclusiveStartFrame());
-				CacheEntry.AccumulatedOffset = Delta * GetAccumulatedOffsetAt(StartTime);
-				CacheEntry.AnimSpaceAccumulatedOffset = Delta * GetAnimSpaceAccumulatedOffsetAt(StartTime);
+				// Use the entry's own start frame for the accumulated offset lookup.
+				// For loop boundary entries within the same section, this chains each
+				// loop iteration on top of the previous one's accumulated offset.
+				CacheEntry.AccumulatedOffset = Delta * GetAccumulatedOffsetAt(Pending.EntryStartFrame);
+				CacheEntry.AnimSpaceAccumulatedOffset = Delta * GetAnimSpaceAccumulatedOffsetAt(Pending.EntryStartFrame);
 			}
 			else
 			{
diff --git a/Engine/Plugins/Experimental/MovieSceneAnimMixer/Source/MovieSceneAnimMixer/Private/Systems/MovieSceneAnimMixerSystem.cpp b/Engine/Plugins/Experimental/MovieSceneAnimMixer/Source/MovieSceneAnimMixer/Private/Systems/MovieSceneAnimMixerSystem.cpp
--- a/Engine/Plugins/Experimental/MovieSceneAnimMixer/Source/MovieSceneAnimMixer/Private/Systems/MovieSceneAnimMixerSystem.cpp
+++ b/Engine/Plugins/Experimental/MovieSceneAnimMixer/Source/MovieSceneAnimMixer/Private/Systems/MovieSceneAnimMixerSystem.cpp
@@ -166,23 +166,23 @@ static TOptional<FAnimNextConvertRootMotionToWorldSpaceTask> BuildSpaceConversio
 				{
 					if (auto* RootMotionSection = TargetDecoration->GetRootMotionSection())
 					{
-						FFrameTime SectionStart(OwnerSection->HasStartFrame()
+						// Use the current evaluation time so loop boundary cache entries
+						// within a section are picked up correctly.
+						FFrameTime AccumLookupTime(OwnerSection->HasStartFrame()
 							? OwnerSection->GetInclusiveStartFrame() : FFrameNumber(0));
-						FFrameNumber StartFrame = SectionStart.FloorToFrame();
-
-						for (const FRootMotionAccumulatedOffset& Offset : RootMotionSection->AccumulatedOffsetCache)
+						if (Context.Linker && MixerEntry.InstanceHandle.IsValid())
 						{
-							if (Offset.StartFrame <= StartFrame)
-							{
-								bUseAccumulatedOffset = true;
-							}
+							const FMovieSceneContext& MovieContext = Context.Linker->GetInstanceRegistry()->GetContext(MixerEntry.InstanceHandle);
+							AccumLookupTime = MovieContext.GetTime();
 						}
+						bUseAccumulatedOffset = !RootMotionSection->AccumulatedOffsetCache.IsEmpty()
+							&& RootMotionSection->AccumulatedOffsetCache[0].StartFrame <= AccumLookupTime;
 
 						if (bUseAccumulatedOffset)
 						{
 							Conversion |= ESpaceConversions::AccumulatedOffset;
 							ResolvedAccumulatedTransform =
-								RootMotionSection->GetAccumulatedOffsetAt(SectionStart);
+								RootMotionSection->GetAccumulatedOffsetAt(AccumLookupTime);
 						}
 					}
 				}
@@ -191,9 +191,11 @@ static TOptional<FAnimNextConvertRootMotionToWorldSpaceTask> BuildSpaceConversio
 				{
 					if (auto* Settings = OwnerSection->FindDecoration<UMovieSceneRootMotionSettingsDecoration>())
 					{
-						if (Settings->bInitialRootTransformValid)
+						if (Settings->bInitialRootTransformValid && Context.Linker && MixerEntry.InstanceHandle.IsValid())
 						{
-							ResolvedInitialRootTransform = Settings->InitialRootTransform;
+							const FMovieSceneContext& MovieContext = Context.Linker->GetInstanceRegistry()->GetContext(MixerEntry.InstanceHandle);
+							ResolvedInitialRootTransform = Settings->GetInitialRootTransformAtTime(
+								MovieContext.GetTime(), MovieContext.GetFrameRate());
 						}
 					}
 				}
diff --git a/Engine/Plugins/Experimental/MovieSceneAnimMixer/Source/MovieSceneAnimMixer/Public/MovieSceneAnimationMixerTrack.h b/Engine/Plugins/Experimental/MovieSceneAnimMixer/Source/MovieSceneAnimMixer/Public/MovieSceneAnimationMixerTrack.h
--- a/Engine/Plugins/Experimental/MovieSceneAnimMixer/Source/MovieSceneAnimMixer/Public/MovieSceneAnimationMixerTrack.h
+++ b/Engine/Plugins/Experimental/MovieSceneAnimMixer/Source/MovieSceneAnimMixer/Public/MovieSceneAnimationMixerTrack.h
@@ -133,6 +133,16 @@ public:
 	MOVIESCENEANIMMIXER_API EMovieSceneRootMotionTransformMode GetRootTransformMode() const;
 	MOVIESCENEANIMMIXER_API EMovieSceneRootMotionSource GetRootMotionSource() const;
 
+	// Returns the first loop boundary time for the owning section, or TOptional empty
+	// if the section doesn't loop. Uses the section's timing transform to find the
+	// actual loop wrap point (preserving sub-frames).
+	MOVIESCENEANIMMIXER_API TOptional<FFrameTime> GetFirstLoopBoundary(const FFrameRate& TickResolution) const;
+
+	// Returns the appropriate initial root transform for the given time - either
+	// InitialRootTransform for the first loop iteration, or LoopStartRootTransform
+	// for subsequent iterations.
+	MOVIESCENEANIMMIXER_API FTransform GetInitialRootTransformAtTime(FFrameTime Time, const FFrameRate& TickResolution) const;
+
 	virtual EMovieSceneChannelProxyType PopulateChannelProxy(FMovieSceneChannelProxyData& OutProxyData) override;
 	virtual void ExtendEntityImpl(UMovieSceneEntitySystemLinker* EntityLinker, const FEntityImportParams& Params, FImportedEntity* OutImportedEntity) override;
 
@@ -203,6 +213,13 @@ public:
 	UPROPERTY()
 	FTransform InitialRootTransform = FTransform::Identity;
 
+	// Root transform at the animation's natural start (frame 0). When the section
+	// is extended left, the first loop starts mid-animation so InitialRootTransform
+	// differs from the animation start. After the first loop wraps, all subsequent
+	// iterations start from this transform instead.
+	UPROPERTY()
+	FTransform LoopStartRootTransform = FTransform::Identity;
+
 	UPROPERTY()
 	bool bInitialRootTransformValid = false;
 
diff --git a/Engine/Plugins/Experimental/MovieSceneAnimMixer/Source/MovieSceneAnimMixer/Public/MovieSceneRootMotionSection.h b/Engine/Plugins/Experimental/MovieSceneAnimMixer/Source/MovieSceneAnimMixer/Public/MovieSceneRootMotionSection.h
--- a/Engine/Plugins/Experimental/MovieSceneAnimMixer/Source/MovieSceneAnimMixer/Public/MovieSceneRootMotionSection.h
+++ b/Engine/Plugins/Experimental/MovieSceneAnimMixer/Source/MovieSceneAnimMixer/Public/MovieSceneRootMotionSection.h
@@ -45,9 +45,9 @@ struct FRootMotionAccumulatedOffset
 {
 	GENERATED_BODY()
 
-	/** Frame at which this accumulated offset takes effect */
+	/** Time at which this accumulated offset takes effect */
 	UPROPERTY()
-	FFrameNumber StartFrame;
+	FFrameTime StartFrame;
 
 	/** The accumulated transform offset at this point (world space) */
 	UPROPERTY()
