Sequencer: Add per-layer weight decoration for animation mixer

Adds a Layer Weight decoration that provides a keyframeable weight
channel per mixer layer. The decoration weight multiplies with the
section weight before inter-layer blending, flowing through both
the mask path and the standard blend path.

#Synced-CL 52530536

diff --git a/Engine/Plugins/Experimental/MovieSceneAnimMixer/Source/MovieSceneAnimMixer/Private/AnimMixerComponentTypes.cpp b/Engine/Plugins/Experimental/MovieSceneAnimMixer/Source/MovieSceneAnimMixer/Private/AnimMixerComponentTypes.cpp
--- a/Engine/Plugins/Experimental/MovieSceneAnimMixer/Source/MovieSceneAnimMixer/Private/AnimMixerComponentTypes.cpp
+++ b/Engine/Plugins/Experimental/MovieSceneAnimMixer/Source/MovieSceneAnimMixer/Private/AnimMixerComponentTypes.cpp
@@ -59,6 +59,7 @@ namespace UE::MovieScene
 
 		Tags.Transition = ComponentRegistry->NewTag(TEXT("Transition"));
 		Tags.SkipPoseWeight = ComponentRegistry->NewTag(TEXT("Skip Pose Weight"));
+		Tags.LayerWeight = ComponentRegistry->NewTag(TEXT("Layer Weight"));
 
 		ComponentRegistry->Factories.DuplicateChildComponent(Priority);
 		ComponentRegistry->Factories.DuplicateChildComponent(Target);
@@ -80,7 +81,7 @@ namespace UE::MovieScene
 		ComponentRegistry->Factories.DefineChildComponent(Task, MixerEntry);
 
 		ComponentRegistry->Factories.DefineChildComponent(Tags.Transition, Tags.Transition);
-
+		ComponentRegistry->Factories.DefineChildComponent(Tags.LayerWeight, Tags.LayerWeight);
 	}
 
 }
\ No newline at end of file
diff --git a/Engine/Plugins/Experimental/MovieSceneAnimMixer/Source/MovieSceneAnimMixer/Private/MovieSceneAnimMixerLayerWeightSection.cpp b/Engine/Plugins/Experimental/MovieSceneAnimMixer/Source/MovieSceneAnimMixer/Private/MovieSceneAnimMixerLayerWeightSection.cpp
new file mode 100644
--- /dev/null
+++ b/Engine/Plugins/Experimental/MovieSceneAnimMixer/Source/MovieSceneAnimMixer/Private/MovieSceneAnimMixerLayerWeightSection.cpp
@@ -0,0 +1,65 @@
+// Copyright Epic Games, Inc. All Rights Reserved.
+
+#include "MovieSceneAnimMixerLayerWeightSection.h"
+
+#include "AnimMixerComponentTypes.h"
+#include "MovieSceneLayerWeightDecoration.h"
+#include "Channels/MovieSceneChannelProxy.h"
+#include "Decorations/IMovieSceneChannelDecoration.h"
+#include "EntitySystem/BuiltInComponentTypes.h"
+#include "Systems/MovieSceneSkeletalAnimationSystem.h"
+
+#include UE_INLINE_GENERATED_CPP_BY_NAME(MovieSceneAnimMixerLayerWeightSection)
+
+UMovieSceneAnimMixerLayerWeightSection::UMovieSceneAnimMixerLayerWeightSection()
+{
+	BlendType = EMovieSceneBlendType::Absolute;
+	SetColorTint(LayerWeightTintColor);
+}
+
+EMovieSceneChannelProxyType UMovieSceneAnimMixerLayerWeightSection::CacheChannelProxy()
+{
+	UMovieSceneLayerWeightDecoration* Decoration = GetTypedOuter<UMovieSceneLayerWeightDecoration>();
+	if (IMovieSceneChannelDecoration* ChannelDecoration = Cast<IMovieSceneChannelDecoration>(Decoration))
+	{
+		FMovieSceneChannelProxyData ChannelProxyData;
+		EMovieSceneChannelProxyType ProxyType = ChannelDecoration->PopulateChannelProxy(ChannelProxyData);
+		ChannelProxy = MakeShared<FMovieSceneChannelProxy>(MoveTemp(ChannelProxyData));
+		return ProxyType;
+	}
+	ChannelProxy = MakeShared<FMovieSceneChannelProxy>();
+	return EMovieSceneChannelProxyType::Dynamic;
+}
+
+void UMovieSceneAnimMixerLayerWeightSection::ImportEntityImpl(UMovieSceneEntitySystemLinker* EntityLinker, const FEntityImportParams& Params,
+	FImportedEntity* OutImportedEntity)
+{
+	using namespace UE::MovieScene;
+
+	const FAnimMixerComponentTypes* AnimMixerComponentTypes = FAnimMixerComponentTypes::Get();
+	FBuiltInComponentTypes* BuiltInComponentTypes = FBuiltInComponentTypes::Get();
+
+	BuildDefaultComponents(EntityLinker, Params, OutImportedEntity);
+
+	if (!OwningLayer)
+	{
+		return;
+	}
+
+	UMovieSceneLayerWeightDecoration* Decoration = GetTypedOuter<UMovieSceneLayerWeightDecoration>();
+	if (!Decoration)
+	{
+		return;
+	}
+
+	const FMovieSceneFloatChannel& WeightChannel = Decoration->GetWeightChannel();
+
+	OutImportedEntity->AddBuilder(
+	FEntityBuilder()
+	.Add(BuiltInComponentTypes->GenericObjectBinding, Params.GetObjectBindingID())
+	.Add(BuiltInComponentTypes->BoundObjectResolver, UMovieSceneSkeletalAnimationSystem::ResolveSkeletalMeshComponentBinding)
+	.Add(AnimMixerComponentTypes->MixerLayer, OwningLayer)
+	.AddConditional(BuiltInComponentTypes->WeightChannel, &WeightChannel, WeightChannel.HasAnyData())
+	.AddTag(AnimMixerComponentTypes->Tags.LayerWeight)
+	.AddMutualComponents());
+}
diff --git a/Engine/Plugins/Experimental/MovieSceneAnimMixer/Source/MovieSceneAnimMixer/Private/MovieSceneAnimMixerModule.cpp b/Engine/Plugins/Experimental/MovieSceneAnimMixer/Source/MovieSceneAnimMixer/Private/MovieSceneAnimMixerModule.cpp
--- a/Engine/Plugins/Experimental/MovieSceneAnimMixer/Source/MovieSceneAnimMixer/Private/MovieSceneAnimMixerModule.cpp
+++ b/Engine/Plugins/Experimental/MovieSceneAnimMixer/Source/MovieSceneAnimMixer/Private/MovieSceneAnimMixerModule.cpp
@@ -9,6 +9,7 @@
 #include "MovieSceneAnimationMixerTrack.h"
 #include "MovieSceneRootMotionTargetDecoration.h"
 #include "MovieSceneAnimationMaskDecoration.h"
+#include "MovieSceneLayerWeightDecoration.h"
 #include "Modules/ModuleInterface.h"
 #include "Modules/ModuleManager.h"
 #include "Sections/MovieSceneSkeletalAnimationSection.h"
@@ -53,6 +54,7 @@ void FMovieSceneAnimMixerModule::StartupModule()
 	IMovieSceneModule::Get().RegisterCompatibleDecoration(UMovieSceneSkeletalAnimationSection::StaticClass(), UMovieSceneRootMotionSettingsDecoration::StaticClass());
 
 	IMovieSceneModule::Get().RegisterCompatibleDecoration(UMovieSceneAnimationMixerLayer::StaticClass(), UMovieSceneAnimationMaskDecoration::StaticClass());
+	IMovieSceneModule::Get().RegisterCompatibleDecoration(UMovieSceneAnimationMixerLayer::StaticClass(), UMovieSceneLayerWeightDecoration::StaticClass());
 
 	IModularFeatures::Get().RegisterModularFeature(IAnimGraphRuntime_SequencerMixerTargetConnector::GetModularFeatureName(), this);
 }
@@ -64,6 +66,7 @@ void FMovieSceneAnimMixerModule::ShutdownModule()
 		IMovieSceneModule::Get().UnregisterCompatibleDecoration(UMovieSceneAnimationMixerTrack::StaticClass(), UMovieSceneRootMotionTargetDecoration::StaticClass());
 		IMovieSceneModule::Get().UnregisterCompatibleDecoration(UMovieSceneSkeletalAnimationSection::StaticClass(), UMovieSceneRootMotionSettingsDecoration::StaticClass());
 		IMovieSceneModule::Get().UnregisterCompatibleDecoration(UMovieSceneAnimationMixerLayer::StaticClass(), UMovieSceneAnimationMaskDecoration::StaticClass());
+		IMovieSceneModule::Get().UnregisterCompatibleDecoration(UMovieSceneAnimationMixerLayer::StaticClass(), UMovieSceneLayerWeightDecoration::StaticClass());
 	}
 
 	IModularFeatures::Get().UnregisterModularFeature(IAnimGraphRuntime_SequencerMixerTargetConnector::GetModularFeatureName(), this);
diff --git a/Engine/Plugins/Experimental/MovieSceneAnimMixer/Source/MovieSceneAnimMixer/Private/MovieSceneLayerWeightDecoration.cpp b/Engine/Plugins/Experimental/MovieSceneAnimMixer/Source/MovieSceneAnimMixer/Private/MovieSceneLayerWeightDecoration.cpp
new file mode 100644
--- /dev/null
+++ b/Engine/Plugins/Experimental/MovieSceneAnimMixer/Source/MovieSceneAnimMixer/Private/MovieSceneLayerWeightDecoration.cpp
@@ -0,0 +1,88 @@
+// Copyright Epic Games, Inc. All Rights Reserved.
+
+#include "MovieSceneLayerWeightDecoration.h"
+
+#include "MovieSceneAnimMixerLayerWeightSection.h"
+#include "MovieSceneAnimationMixerLayer.h"
+#include "Channels/MovieSceneChannelProxy.h"
+
+#define LOCTEXT_NAMESPACE "MovieSceneLayerWeightDecoration"
+
+UMovieSceneLayerWeightDecoration::UMovieSceneLayerWeightDecoration(const FObjectInitializer& ObjInit)
+	: Super(ObjInit)
+{
+	WeightChannel.SetDefault(1.0f);
+}
+
+EMovieSceneChannelProxyType UMovieSceneLayerWeightDecoration::PopulateChannelProxy(FMovieSceneChannelProxyData& OutProxyData)
+{
+#if WITH_EDITOR
+	FMovieSceneChannelMetaData MetaData(TEXT("WeightChannel"), LOCTEXT("Weight", "Weight"));
+	MetaData.WeakOwningObject = this;
+	OutProxyData.Add(WeightChannel, MetaData, TMovieSceneExternalValue<float>());
+#else
+	OutProxyData.Add(WeightChannel);
+#endif
+	return EMovieSceneChannelProxyType::Dynamic;
+}
+
+UMovieSceneSection* UMovieSceneLayerWeightDecoration::CreateNewSection()
+{
+	return NewObject<UMovieSceneAnimMixerLayerWeightSection>(this, UMovieSceneAnimMixerLayerWeightSection::StaticClass(), NAME_None, RF_Transactional);
+}
+
+void UMovieSceneLayerWeightDecoration::AddSection(UMovieSceneSection* InSection)
+{
+	UMovieSceneAnimMixerLayerWeightSection* NewWeightSection = Cast<UMovieSceneAnimMixerLayerWeightSection>(InSection);
+	if (!NewWeightSection || WeightSections.Num() > 0)
+	{
+		return;
+	}
+
+	Modify();
+
+	WeightSections.Add(NewWeightSection);
+
+	MarkAsChanged();
+#if WITH_EDITOR
+	MarkStructureChanged();
+#endif
+
+	UMovieSceneAnimationMixerLayer* Layer = GetTypedOuter<UMovieSceneAnimationMixerLayer>();
+	NewWeightSection->SetAnimMixerLayer(Layer);
+
+	if (Layer)
+	{
+		NewWeightSection->SetRowIndex(Layer->GetLayerIndex());
+	}
+}
+
+void UMovieSceneLayerWeightDecoration::RemoveSection(UMovieSceneSection& SectionToRemove)
+{
+	Modify();
+	WeightSections.RemoveAll([&SectionToRemove](const UMovieSceneSection* InSection){ return InSection == &SectionToRemove; });
+	MarkAsChanged();
+#if WITH_EDITOR
+	MarkStructureChanged();
+#endif
+}
+
+#if WITH_EDITOR
+void UMovieSceneLayerWeightDecoration::PostEditUndo()
+{
+	Super::PostEditUndo();
+
+	UMovieSceneAnimationMixerLayer* Layer = GetTypedOuter<UMovieSceneAnimationMixerLayer>();
+	for (TObjectPtr<UMovieSceneSection> Section : WeightSections)
+	{
+		if (UMovieSceneAnimMixerLayerWeightSection* WeightSection = Cast<UMovieSceneAnimMixerLayerWeightSection>(Section))
+		{
+			WeightSection->SetAnimMixerLayer(Layer);
+		}
+	}
+
+	MarkStructureChanged();
+}
+#endif
+
+#undef LOCTEXT_NAMESPACE
diff --git a/Engine/Plugins/Experimental/MovieSceneAnimMixer/Source/MovieSceneAnimMixer/Private/Systems/MovieSceneAnimMixerMaskSystem.cpp b/Engine/Plugins/Experimental/MovieSceneAnimMixer/Source/MovieSceneAnimMixer/Private/Systems/MovieSceneAnimMixerMaskSystem.cpp
--- a/Engine/Plugins/Experimental/MovieSceneAnimMixer/Source/MovieSceneAnimMixer/Private/Systems/MovieSceneAnimMixerMaskSystem.cpp
+++ b/Engine/Plugins/Experimental/MovieSceneAnimMixer/Source/MovieSceneAnimMixer/Private/Systems/MovieSceneAnimMixerMaskSystem.cpp
@@ -256,15 +256,50 @@ struct FGatherMaskData
 			FLayerBlendTaskBuilder Builder;
 			if (Accum.NumContributingSections > 0)
 			{
-				// Pre-build the task with ScaleFactor=1.0. The builder adjusts
-				// ScaleFactor to the actual layer weight each frame 
+				// Build the task with ScaleFactor=1.0. The builder reconstructs it
+				// each frame with the actual layer weight, pre-multiplying per-bone,
+				// per-curve and per-attribute weights so the blend correctly scales
+				// the mask contribution.
 				TSharedPtr<FAnimNextBlendKeyframePerBoneWithScaleTask> Task =
 					StaticCastSharedPtr<FAnimNextBlendKeyframePerBoneWithScaleTask>(Accum.BuildBlendTask(1.0f));
 				if (Task)
 				{
-					Builder.BuildTask = [Task](double LayerWeight) -> TSharedPtr<FAnimNextEvaluationTask>
+					auto* AccumulationsMap = Accumulations;
+					FMaskLayerKey AccumKey = Key;
+					Builder.BuildTask = [Task, AccumulationsMap, AccumKey](double LayerWeight) -> TSharedPtr<FAnimNextEvaluationTask>
 					{
-						Task->ScaleFactor = static_cast<float>(LayerWeight);
+						FMaskLayerAccumulation* Accum = AccumulationsMap->Find(AccumKey);
+						if (!Accum)
+						{
+							return Task;
+						}
+
+						const float Scale = static_cast<float>(LayerWeight);
+
+						Accum->TaskBoneWeights.Reset();
+						Accum->TaskCurveWeights.Empty();
+						Accum->TaskAttributeWeights.Reset();
+
+						for (int32 i = 0; i < Accum->AccumulationBoneWeights.Num(); ++i)
+						{
+							Accum->TaskBoneWeights.Add((1.0f - Accum->AccumulationBoneWeights[i]) * Scale);
+						}
+
+						for (const TTuple<FName, float>& CurveWeight : Accum->CurveWeightsMap)
+						{
+							Accum->TaskCurveWeights.Add(UE::Anim::FCurveElement(CurveWeight.Key, (1.0f - CurveWeight.Value) * Scale));
+						}
+
+						for (const TTuple<UE::Anim::FAttributeId, float>& AttrWeight : Accum->AttributeWeightsMap)
+						{
+							Accum->TaskAttributeWeights.Add(FBlendProfileStandaloneCachedData::FMaskedAttributeWeight(
+								AttrWeight.Key, (1.0f - AttrWeight.Value) * Scale));
+						}
+
+						*Task = FAnimNextBlendKeyframePerBoneWithScaleTask::Make(
+							Accum->MaskSkeleton.Get(), &Accum->TaskBoneWeights,
+							&Accum->TaskCurveWeights, &Accum->TaskAttributeWeights, Scale);
+
 						return Task;
 					};
 				}
diff --git a/Engine/Plugins/Experimental/MovieSceneAnimMixer/Source/MovieSceneAnimMixer/Private/Systems/MovieSceneAnimMixerSystem.cpp b/Engine/Plugins/Experimental/MovieSceneAnimMixer/Source/MovieSceneAnimMixer/Private/Systems/MovieSceneAnimMixerSystem.cpp
--- a/Engine/Plugins/Experimental/MovieSceneAnimMixer/Source/MovieSceneAnimMixer/Private/Systems/MovieSceneAnimMixerSystem.cpp
+++ b/Engine/Plugins/Experimental/MovieSceneAnimMixer/Source/MovieSceneAnimMixer/Private/Systems/MovieSceneAnimMixerSystem.cpp
@@ -89,6 +89,8 @@ struct FProgramBuildContext
 	// Layer blend task builders, keyed by MixerLayer for this specific mixer.
 	// Built from entities with {BoundObjectKey, MixerLayer, LayerBlendTask} matching this mixer's BoundObjectKey.
 	const TMap<TWeakObjectPtr<UMovieSceneAnimationMixerLayer>, FLayerBlendTaskBuilder>* LayerBlendTaskMap = nullptr;
+	// Decoration-provided layer weights. Multiplied into LayerWeight before inter-layer blending.
+	const TMap<TWeakObjectPtr<UMovieSceneAnimationMixerLayer>, double>* DecorationLayerWeightMap = nullptr;
 };
 
 // Returns true if the entry passes the filter (i.e. should be considered for inclusion)
@@ -671,6 +673,16 @@ static TSharedPtr<UE::UAF::FEvaluationProgram> BuildProgramFromEntries(
 
 		// ---- Level 2: Inter-layer blend ----
 		// Blend this layer's result with what's below on the stack.
+
+		// Apply decoration-provided layer weight (e.g. from a Layer Weight decoration)
+		if (Layer.IsValid() && Context.DecorationLayerWeightMap)
+		{
+			if (const double* DecWeight = Context.DecorationLayerWeightMap->Find(Layer))
+			{
+				LayerWeight *= *DecWeight;
+			}
+		}
+
 		TSharedPtr<FAnimNextEvaluationTask> LayerBlendTask;
 		if (Layer.IsValid() && Context.LayerBlendTaskMap)
 		{
@@ -780,6 +792,9 @@ namespace UE::MovieScene
 		// Layer blend task builders grouped by BoundObjectKey -> {MixerLayer -> Builder}
 		TMap<FObjectKey, TMap<TWeakObjectPtr<UMovieSceneAnimationMixerLayer>, FLayerBlendTaskBuilder>> LayerBlendBuilders;
 
+		// Decoration-provided layer weights grouped by BoundObjectKey -> {MixerLayer -> Weight}
+		TMap<FObjectKey, TMap<TWeakObjectPtr<UMovieSceneAnimationMixerLayer>, double>> DecorationLayerWeights;
+
 		FEvaluateAnimMixers(TMap<FMovieSceneAnimMixerKey, TSharedPtr<FMovieSceneAnimMixer>>* InMixers, TMap<FObjectKey, TWeakPtr<FMovieSceneMixerRootMotionComponentData>>* InRootMotions, UMovieSceneEntitySystemLinker* InLinker, UMovieSceneAnimMixerSystem* InMixerSystem)
 			: Mixers(InMixers)
 			, RootMotions(InRootMotions)
@@ -810,6 +825,25 @@ namespace UE::MovieScene
 					}
 				}
 			);
+
+			// Build decoration layer weight lookup from layer weight section entities.
+			// These are tagged with LayerWeight and have their WeightResult evaluated
+			// by the standard weight/easing system.
+			FEntityTaskBuilder()
+			.Read(BuiltInComponents->BoundObjectKey)
+			.Read(AnimMixerComponents->MixerLayer)
+			.Read(BuiltInComponents->WeightResult)
+			.FilterAll({ AnimMixerComponents->Tags.LayerWeight })
+			.FilterNone({ BuiltInComponents->Tags.NeedsUnlink })
+			.Iterate_PerEntity(&Linker->EntityManager,
+				[this](FObjectKey BoundObjectKey, TObjectPtr<UMovieSceneAnimationMixerLayer> MixerLayer, double Weight)
+				{
+					if (MixerLayer)
+					{
+						DecorationLayerWeights.FindOrAdd(BoundObjectKey).Add(MixerLayer, Weight);
+					}
+				}
+			);
 		}
 
 		// Get an InstanceHandle from the mixer's entries (any valid entry will do).
@@ -856,6 +890,9 @@ namespace UE::MovieScene
 				LayerBlendBuilders.Find(MeshComponent.GetObject());
 			BuildContext.LayerBlendTaskMap = BlendMap;
 
+			// Look up decoration layer weights for this bound object
+			BuildContext.DecorationLayerWeightMap = DecorationLayerWeights.Find(MeshComponent.GetObject());
+
 			// Apply bake filter from the bake target if present.
 			FProgramBuildFilter BuildFilter;
 			if (const FMovieSceneMixerBakeTarget* BakeTarget = Target.GetPtr<FMovieSceneMixerBakeTarget>())
diff --git a/Engine/Plugins/Experimental/MovieSceneAnimMixer/Source/MovieSceneAnimMixer/Public/AnimMixerComponentTypes.h b/Engine/Plugins/Experimental/MovieSceneAnimMixer/Source/MovieSceneAnimMixer/Public/AnimMixerComponentTypes.h
--- a/Engine/Plugins/Experimental/MovieSceneAnimMixer/Source/MovieSceneAnimMixer/Public/AnimMixerComponentTypes.h
+++ b/Engine/Plugins/Experimental/MovieSceneAnimMixer/Source/MovieSceneAnimMixer/Public/AnimMixerComponentTypes.h
@@ -74,6 +74,7 @@ namespace UE::MovieScene
 		{
 			FComponentTypeID Transition;
 			FComponentTypeID SkipPoseWeight;
+			FComponentTypeID LayerWeight;
 		} Tags;
 
 
diff --git a/Engine/Plugins/Experimental/MovieSceneAnimMixer/Source/MovieSceneAnimMixer/Public/MovieSceneAnimMixerLayerWeightSection.h b/Engine/Plugins/Experimental/MovieSceneAnimMixer/Source/MovieSceneAnimMixer/Public/MovieSceneAnimMixerLayerWeightSection.h
new file mode 100644
--- /dev/null
+++ b/Engine/Plugins/Experimental/MovieSceneAnimMixer/Source/MovieSceneAnimMixer/Public/MovieSceneAnimMixerLayerWeightSection.h
@@ -0,0 +1,43 @@
+// Copyright Epic Games, Inc. All Rights Reserved.
+
+#pragma once
+
+#include "MovieSceneAnimationMixerItemInterface.h"
+#include "MovieSceneAnimationMixerLayer.h"
+#include "MovieSceneSection.h"
+#include "EntitySystem/IMovieSceneEntityProvider.h"
+
+#include "MovieSceneAnimMixerLayerWeightSection.generated.h"
+
+// Section for the layer weight decoration. Imports an entity tagged with
+// LayerWeight so the mixer can read the evaluated weight channel from
+// the owning decoration.
+UCLASS(MinimalAPI, DisplayName="Layer Weight")
+class UMovieSceneAnimMixerLayerWeightSection
+	: public UMovieSceneSection
+	, public IMovieSceneEntityProvider
+	, public IMovieSceneAnimationMixerItemInterface
+{
+public:
+
+	UMovieSceneAnimMixerLayerWeightSection();
+
+	GENERATED_BODY()
+
+	void SetAnimMixerLayer(UMovieSceneAnimationMixerLayer* InLayer) { OwningLayer = InLayer; }
+
+#if WITH_EDITOR
+	virtual bool SupportsVerticalDragging() const override { return false; }
+	virtual bool IsVisibleInAddSectionMenu() const override { return false; }
+#endif
+
+private:
+
+	UPROPERTY()
+	TObjectPtr<UMovieSceneAnimationMixerLayer> OwningLayer;
+
+	virtual EMovieSceneChannelProxyType CacheChannelProxy() override;
+	virtual void ImportEntityImpl(UMovieSceneEntitySystemLinker* EntityLinker, const FEntityImportParams& Params, FImportedEntity* OutImportedEntity) override;
+
+	FColor LayerWeightTintColor = FColor(25, 80, 110);
+};
diff --git a/Engine/Plugins/Experimental/MovieSceneAnimMixer/Source/MovieSceneAnimMixer/Public/MovieSceneLayerWeightDecoration.h b/Engine/Plugins/Experimental/MovieSceneAnimMixer/Source/MovieSceneAnimMixer/Public/MovieSceneLayerWeightDecoration.h
new file mode 100644
--- /dev/null
+++ b/Engine/Plugins/Experimental/MovieSceneAnimMixer/Source/MovieSceneAnimMixer/Public/MovieSceneLayerWeightDecoration.h
@@ -0,0 +1,49 @@
+// Copyright Epic Games, Inc. All Rights Reserved.
+
+#pragma once
+
+#include "MovieSceneSignedObject.h"
+#include "Channels/MovieSceneFloatChannel.h"
+#include "Decorations/IMovieSceneChannelDecoration.h"
+#include "Decorations/IMovieSceneSectionProviderDecoration.h"
+
+#include "MovieSceneLayerWeightDecoration.generated.h"
+
+class UMovieSceneAnimationMixerLayer;
+
+// Decoration that adds a keyframeable weight channel to a mixer layer.
+// The weight (default 1.0) multiplies with the section weight before
+// inter-layer blending, scaling the layer's contribution to the final pose.
+UCLASS(MinimalAPI, DisplayName="Layer Weight")
+class UMovieSceneLayerWeightDecoration
+	: public UMovieSceneSignedObject
+	, public IMovieSceneSectionProviderDecoration
+	, public IMovieSceneChannelDecoration
+{
+public:
+	GENERATED_BODY()
+
+	UMovieSceneLayerWeightDecoration(const FObjectInitializer& ObjInit);
+
+	// IMovieSceneSectionProviderDecoration
+	virtual TArrayView<TObjectPtr<UMovieSceneSection>> GetSections() override { return WeightSections; }
+	virtual UMovieSceneSection* CreateNewSection() override;
+	MOVIESCENEANIMMIXER_API virtual void AddSection(UMovieSceneSection* InSection) override;
+	MOVIESCENEANIMMIXER_API virtual void RemoveSection(UMovieSceneSection& SectionToRemove) override;
+	MOVIESCENEANIMMIXER_API virtual bool SupportsMultipleSections() override { return false; }
+
+	// IMovieSceneChannelDecoration
+	virtual EMovieSceneChannelProxyType PopulateChannelProxy(FMovieSceneChannelProxyData& OutProxyData) override;
+
+	const FMovieSceneFloatChannel& GetWeightChannel() const { return WeightChannel; }
+
+	UPROPERTY()
+	FMovieSceneFloatChannel WeightChannel;
+
+	UPROPERTY()
+	TArray<TObjectPtr<UMovieSceneSection>> WeightSections;
+
+#if WITH_EDITOR
+	virtual void PostEditUndo() override;
+#endif
+};
diff --git a/Engine/Plugins/Experimental/MovieSceneAnimMixer/Source/MovieSceneAnimMixerEditor/Private/LayerWeightDecorationEditor.cpp b/Engine/Plugins/Experimental/MovieSceneAnimMixer/Source/MovieSceneAnimMixerEditor/Private/LayerWeightDecorationEditor.cpp
new file mode 100644
--- /dev/null
+++ b/Engine/Plugins/Experimental/MovieSceneAnimMixer/Source/MovieSceneAnimMixerEditor/Private/LayerWeightDecorationEditor.cpp
@@ -0,0 +1,18 @@
+// Copyright Epic Games, Inc. All Rights Reserved.
+
+#include "LayerWeightDecorationEditor.h"
+
+#include "MovieSceneLayerWeightDecoration.h"
+
+namespace UE::Sequencer
+{
+	UClass* FLayerWeightDecorationEditor::GetDecorationClass() const
+	{
+		return UMovieSceneLayerWeightDecoration::StaticClass();
+	}
+
+	TSharedPtr<ISequencerSection> FLayerWeightDecorationEditor::MakeSectionInterface(UMovieSceneSection& Section, UMovieSceneTrack& OwningTrack, FGuid ObjectBinding)
+	{
+		return MakeShared<FSequencerSection>(Section);
+	}
+}
diff --git a/Engine/Plugins/Experimental/MovieSceneAnimMixer/Source/MovieSceneAnimMixerEditor/Private/LayerWeightDecorationEditor.h b/Engine/Plugins/Experimental/MovieSceneAnimMixer/Source/MovieSceneAnimMixerEditor/Private/LayerWeightDecorationEditor.h
new file mode 100644
--- /dev/null
+++ b/Engine/Plugins/Experimental/MovieSceneAnimMixer/Source/MovieSceneAnimMixerEditor/Private/LayerWeightDecorationEditor.h
@@ -0,0 +1,18 @@
+// Copyright Epic Games, Inc. All Rights Reserved.
+
+#pragma once
+#include "ISequencerDecorationEditor.h"
+#include "ISequencerSection.h"
+
+namespace UE::Sequencer
+{
+
+class FLayerWeightDecorationEditor : public ISequencerDecorationEditor
+{
+public:
+	virtual UClass* GetDecorationClass() const override;
+
+	virtual TSharedPtr<ISequencerSection> MakeSectionInterface(UMovieSceneSection& Section, UMovieSceneTrack& OwningTrack, FGuid ObjectBinding) override;
+};
+
+}
diff --git a/Engine/Plugins/Experimental/MovieSceneAnimMixer/Source/MovieSceneAnimMixerEditor/Private/LayerWeightDecorationMenuProvider.cpp b/Engine/Plugins/Experimental/MovieSceneAnimMixer/Source/MovieSceneAnimMixerEditor/Private/LayerWeightDecorationMenuProvider.cpp
new file mode 100644
--- /dev/null
+++ b/Engine/Plugins/Experimental/MovieSceneAnimMixer/Source/MovieSceneAnimMixerEditor/Private/LayerWeightDecorationMenuProvider.cpp
@@ -0,0 +1,65 @@
+// Copyright Epic Games, Inc. All Rights Reserved.
+
+#include "LayerWeightDecorationMenuProvider.h"
+
+#include "ISequencer.h"
+#include "MovieScene.h"
+#include "MovieSceneLayerWeightDecoration.h"
+#include "Decorations/MovieSceneDecorationContainer.h"
+#include "Framework/MultiBox/MultiBoxBuilder.h"
+#include "ScopedTransaction.h"
+
+#define LOCTEXT_NAMESPACE "LayerWeightDecorationMenuProvider"
+
+UClass* FLayerWeightDecorationMenuProvider::GetHandledDecorationClass() const
+{
+	return UMovieSceneLayerWeightDecoration::StaticClass();
+}
+
+void FLayerWeightDecorationMenuProvider::PopulateAddDecorationMenu(FMenuBuilder& MenuBuilder, TObjectPtr<UObject> BoundObject,
+	UMovieSceneDecorationContainerObject* DecorationContainer, TWeakPtr<ISequencer> InSequencer)
+{
+	MenuBuilder.AddMenuEntry(
+		LOCTEXT("LayerWeightDecoration", "Layer Weight"),
+		LOCTEXT("LayerWeightDecorationTooltip", "Add a keyframeable weight channel to this layer."),
+		FSlateIcon(),
+		FUIAction(FExecuteAction::CreateRaw(this, &FLayerWeightDecorationMenuProvider::OnAddLayerWeight, DecorationContainer, InSequencer))
+	);
+}
+
+void FLayerWeightDecorationMenuProvider::OnAddLayerWeight(UMovieSceneDecorationContainerObject* DecorationContainer, TWeakPtr<ISequencer> WeakSequencer)
+{
+	if (!DecorationContainer)
+	{
+		return;
+	}
+
+	TSharedPtr<ISequencer> Sequencer = WeakSequencer.Pin();
+	if (!Sequencer)
+	{
+		return;
+	}
+
+	FScopedTransaction Transaction(LOCTEXT("AddLayerWeight", "Add Layer Weight"));
+
+	UMovieSceneLayerWeightDecoration* Decoration = DecorationContainer->GetOrCreateDecoration<UMovieSceneLayerWeightDecoration>();
+	if (!Decoration)
+	{
+		return;
+	}
+
+	UMovieSceneSection* NewSection = Decoration->CreateNewSection();
+	if (!NewSection)
+	{
+		return;
+	}
+
+	const TRange<FFrameNumber> PlaybackRange = Sequencer->GetFocusedMovieSceneSequence()->GetMovieScene()->GetPlaybackRange();
+	NewSection->SetRange(PlaybackRange);
+
+	Decoration->AddSection(NewSection);
+
+	Sequencer->NotifyMovieSceneDataChanged(EMovieSceneDataChangeType::MovieSceneStructureItemAdded);
+}
+
+#undef LOCTEXT_NAMESPACE
diff --git a/Engine/Plugins/Experimental/MovieSceneAnimMixer/Source/MovieSceneAnimMixerEditor/Private/LayerWeightDecorationMenuProvider.h b/Engine/Plugins/Experimental/MovieSceneAnimMixer/Source/MovieSceneAnimMixerEditor/Private/LayerWeightDecorationMenuProvider.h
new file mode 100644
--- /dev/null
+++ b/Engine/Plugins/Experimental/MovieSceneAnimMixer/Source/MovieSceneAnimMixerEditor/Private/LayerWeightDecorationMenuProvider.h
@@ -0,0 +1,21 @@
+// Copyright Epic Games, Inc. All Rights Reserved.
+
+#pragma once
+
+#include "Decorations/IMovieSceneDecorationMenuProvider.h"
+
+// Decoration menu provider for layer weight (UMovieSceneLayerWeightDecoration).
+// Adds a "Layer Weight" menu entry that creates a layer weight section spanning
+// the current playback range.
+class FLayerWeightDecorationMenuProvider : public IMovieSceneDecorationMenuProvider
+{
+public:
+
+	virtual UClass* GetHandledDecorationClass() const override;
+	virtual void PopulateAddDecorationMenu(FMenuBuilder& MenuBuilder, TObjectPtr<UObject> BoundObject, UMovieSceneDecorationContainerObject* DecorationContainer, TWeakPtr<ISequencer> InSequencer) override;
+	virtual int32 GetDecorationMenuPriority() const override { return 50; }
+
+private:
+
+	void OnAddLayerWeight(UMovieSceneDecorationContainerObject* DecorationContainer, TWeakPtr<ISequencer> WeakSequencer);
+};
diff --git a/Engine/Plugins/Experimental/MovieSceneAnimMixer/Source/MovieSceneAnimMixerEditor/Private/MovieSceneAnimMixerEditorModule.cpp b/Engine/Plugins/Experimental/MovieSceneAnimMixer/Source/MovieSceneAnimMixerEditor/Private/MovieSceneAnimMixerEditorModule.cpp
--- a/Engine/Plugins/Experimental/MovieSceneAnimMixer/Source/MovieSceneAnimMixerEditor/Private/MovieSceneAnimMixerEditorModule.cpp
+++ b/Engine/Plugins/Experimental/MovieSceneAnimMixer/Source/MovieSceneAnimMixerEditor/Private/MovieSceneAnimMixerEditorModule.cpp
@@ -4,6 +4,8 @@
 
 #include "AnimBlueprintExtension_SequencerMixerTarget.h"
 #include "AnimMixerDecorationMenuProviders.h"
+#include "LayerWeightDecorationEditor.h"
+#include "LayerWeightDecorationMenuProvider.h"
 #include "IMovieSceneRootMotionOffsetProvider.h"
 #include "AnimMixerTargetMenuProviders.h"
 #include "AnimTransitionSectionInterface.h"
@@ -257,9 +259,15 @@ namespace UE::MovieScene
 		AnimMaskingDecorationEditor = MakeUnique<FMovieSceneAnimationMaskDecorationEditor>();
 		IModularFeatures::Get().RegisterModularFeature(ISequencerDecorationEditor::GetModularFeatureName(), AnimMaskingDecorationEditor.Get());
 
+		LayerWeightDecorationEditor = MakeUnique<UE::Sequencer::FLayerWeightDecorationEditor>();
+		IModularFeatures::Get().RegisterModularFeature(ISequencerDecorationEditor::GetModularFeatureName(), LayerWeightDecorationEditor.Get());
+
 		// Register decoration menu providers
 		MaskDecorationMenuProvider = MakeUnique<FMaskDecorationMenuProvider>();
 		IModularFeatures::Get().RegisterModularFeature(IMovieSceneDecorationMenuProvider::GetModularFeatureName(), MaskDecorationMenuProvider.Get());
+
+		LayerWeightDecorationMenuProvider = MakeUnique<FLayerWeightDecorationMenuProvider>();
+		IModularFeatures::Get().RegisterModularFeature(IMovieSceneDecorationMenuProvider::GetModularFeatureName(), LayerWeightDecorationMenuProvider.Get());
 		IModularFeatures::Get().RegisterModularFeature(IMovieSceneRootMotionOffsetProvider::GetModularFeatureName(), &AnimMixerRootMotionProvider);
 	}
 
@@ -281,12 +289,20 @@ namespace UE::MovieScene
 		{
 			IModularFeatures::Get().UnregisterModularFeature(ISequencerDecorationEditor::GetModularFeatureName(), AnimMaskingDecorationEditor.Get());
 		}
+		if (LayerWeightDecorationEditor)
+		{
+			IModularFeatures::Get().UnregisterModularFeature(ISequencerDecorationEditor::GetModularFeatureName(), LayerWeightDecorationEditor.Get());
+		}
 
 		// Unregister decoration menu providers
 		if (MaskDecorationMenuProvider)
 		{
 			IModularFeatures::Get().UnregisterModularFeature(IMovieSceneDecorationMenuProvider::GetModularFeatureName(), MaskDecorationMenuProvider.Get());
 		}
+		if (LayerWeightDecorationMenuProvider)
+		{
+			IModularFeatures::Get().UnregisterModularFeature(IMovieSceneDecorationMenuProvider::GetModularFeatureName(), LayerWeightDecorationMenuProvider.Get());
+		}
 
 		// Unregister target menu providers
 		if (AutomaticTargetMenuProvider)
diff --git a/Engine/Plugins/Experimental/MovieSceneAnimMixer/Source/MovieSceneAnimMixerEditor/Private/MovieSceneAnimMixerEditorModule.h b/Engine/Plugins/Experimental/MovieSceneAnimMixer/Source/MovieSceneAnimMixerEditor/Private/MovieSceneAnimMixerEditorModule.h
--- a/Engine/Plugins/Experimental/MovieSceneAnimMixer/Source/MovieSceneAnimMixerEditor/Private/MovieSceneAnimMixerEditorModule.h
+++ b/Engine/Plugins/Experimental/MovieSceneAnimMixer/Source/MovieSceneAnimMixerEditor/Private/MovieSceneAnimMixerEditorModule.h
@@ -13,11 +13,13 @@ class FAnimInstanceTargetMenuProvider;
 class FAnimBlueprintTargetMenuProvider;
 class FAnimNextInjectionTargetMenuProvider;
 class FMaskDecorationMenuProvider;
+class FLayerWeightDecorationMenuProvider;
 
 namespace UE::Sequencer
 {
 	class FRootMotionTargetDecorationEditor;
 	class FRootMotionSettingsDecorationEditor;
+	class FLayerWeightDecorationEditor;
 }
 
 namespace UE::MovieScene
@@ -52,9 +54,11 @@ namespace UE::MovieScene
 		TUniquePtr<UE::Sequencer::FRootMotionTargetDecorationEditor> RootMotionTargetDecorationEditor;
 		TUniquePtr<UE::Sequencer::FRootMotionSettingsDecorationEditor> RootMotionSettingsDecorationEditor;
 		TUniquePtr<Sequencer::FMovieSceneAnimationMaskDecorationEditor> AnimMaskingDecorationEditor;
+		TUniquePtr<UE::Sequencer::FLayerWeightDecorationEditor> LayerWeightDecorationEditor;
 
 		// Decoration menu providers
 		TUniquePtr<FMaskDecorationMenuProvider> MaskDecorationMenuProvider;
+		TUniquePtr<FLayerWeightDecorationMenuProvider> LayerWeightDecorationMenuProvider;
 		FAnimMixerRootMotionOffsetProvider AnimMixerRootMotionProvider;
 	};
 }
