Sequencer: Fix linux compile

#Synced-CL 52530536

diff --git a/Engine/Plugins/Experimental/MovieSceneAnimMixer/Source/MovieSceneAnimMixer/Private/Systems/MovieSceneAnimMixerMaskSystem.cpp b/Engine/Plugins/Experimental/MovieSceneAnimMixer/Source/MovieSceneAnimMixer/Private/Systems/MovieSceneAnimMixerMaskSystem.cpp
--- a/Engine/Plugins/Experimental/MovieSceneAnimMixer/Source/MovieSceneAnimMixer/Private/Systems/MovieSceneAnimMixerMaskSystem.cpp
+++ b/Engine/Plugins/Experimental/MovieSceneAnimMixer/Source/MovieSceneAnimMixer/Private/Systems/MovieSceneAnimMixerMaskSystem.cpp
@@ -268,37 +268,37 @@ struct FGatherMaskData
 					FMaskLayerKey AccumKey = Key;
 					Builder.BuildTask = [Task, AccumulationsMap, AccumKey](double LayerWeight) -> TSharedPtr<FAnimNextEvaluationTask>
 					{
-						FMaskLayerAccumulation* Accum = AccumulationsMap->Find(AccumKey);
-						if (!Accum)
+						FMaskLayerAccumulation* AccumData = AccumulationsMap->Find(AccumKey);
+						if (!AccumData)
 						{
 							return Task;
 						}
 
 						const float Scale = static_cast<float>(LayerWeight);
 
-						Accum->TaskBoneWeights.Reset();
-						Accum->TaskCurveWeights.Empty();
-						Accum->TaskAttributeWeights.Reset();
+						AccumData->TaskBoneWeights.Reset();
+						AccumData->TaskCurveWeights.Empty();
+						AccumData->TaskAttributeWeights.Reset();
 
-						for (int32 i = 0; i < Accum->AccumulationBoneWeights.Num(); ++i)
+						for (int32 i = 0; i < AccumData->AccumulationBoneWeights.Num(); ++i)
 						{
-							Accum->TaskBoneWeights.Add((1.0f - Accum->AccumulationBoneWeights[i]) * Scale);
+							AccumData->TaskBoneWeights.Add((1.0f - AccumData->AccumulationBoneWeights[i]) * Scale);
 						}
 
-						for (const TTuple<FName, float>& CurveWeight : Accum->CurveWeightsMap)
+						for (const TTuple<FName, float>& CurveWeight : AccumData->CurveWeightsMap)
 						{
-							Accum->TaskCurveWeights.Add(UE::Anim::FCurveElement(CurveWeight.Key, (1.0f - CurveWeight.Value) * Scale));
+							AccumData->TaskCurveWeights.Add(UE::Anim::FCurveElement(CurveWeight.Key, (1.0f - CurveWeight.Value) * Scale));
 						}
 
-						for (const TTuple<UE::Anim::FAttributeId, float>& AttrWeight : Accum->AttributeWeightsMap)
+						for (const TTuple<UE::Anim::FAttributeId, float>& AttrWeight : AccumData->AttributeWeightsMap)
 						{
-							Accum->TaskAttributeWeights.Add(FBlendProfileStandaloneCachedData::FMaskedAttributeWeight(
+							AccumData->TaskAttributeWeights.Add(FBlendProfileStandaloneCachedData::FMaskedAttributeWeight(
 								AttrWeight.Key, (1.0f - AttrWeight.Value) * Scale));
 						}
 
 						*Task = FAnimNextBlendKeyframePerBoneWithScaleTask::Make(
-							Accum->MaskSkeleton.Get(), &Accum->TaskBoneWeights,
-							&Accum->TaskCurveWeights, &Accum->TaskAttributeWeights, Scale);
+							AccumData->MaskSkeleton.Get(), &AccumData->TaskBoneWeights,
+							&AccumData->TaskCurveWeights, &AccumData->TaskAttributeWeights, Scale);
 
 						return Task;
 					};
