[Backout] - CL52569948
#rnx
Original CL Desc
-----------------------------------------------------------------
[EDA] Add SequencerTools toolset for AI-driven sequence creation and editing

Adds a comprehensive Python toolset (220 tools) to the AnimationAssistantToolset plugin that wraps Sequencer's scripting APIs for EDA agents. This enables agents to create and edit Level Sequences, control playback, manage bindings/tracks/sections/keys, operate the Anim Mixer, drive Control Rig, and inspect the outliner through the EDA toolset system.

Python toolset (sequencer.py):
- Sequence lifecycle: create, open, close, refresh
- Playback: play, pause, set/get playhead, speed, loop mode
- Sequence properties: display rate, tick resolution, evaluation type, clock source, playback/view/work ranges, marked frames, playback range lock
- Bindings: add actors, create spawnables, create cameras, find/rename/remove bindings
- Tracks: add/remove/find tracks on bindings and sequences, event track trigger/repeater sections
- Sections: add/remove/set range, easing, blend type, completion mode, bounded/infinite ranges, camera cut bindings
- Keys/Channels: add/get/remove keys (float/bool/int/string), set/get defaults, channel selection, evaluate
- Selection: tracks, sections, bindings, folders, channels, outliner nodes
- Folders: create, remove, organize tracks and bindings
- Sub-sequences: focus into/out of sub-sequences, hierarchy query
- Outliner: tree snapshot, children, expansion, mute, solo, deactivated, locked, pinned states
- Anim Mixer: layers (add/insert/rename/query), animations, transitions (query/change type), decorations (generic discovery/add/remove)
- Control Rig: find/create CR tracks, get/set values for all control types (bool, float, int, position, rotator, scale, vector2d, euler transform, transform), world transforms, actor transforms, space switching (set/delete/move/bake), keying, masking, layering, selection, mirroring, tweening, snapping, FK detection
- Import/Export: FBX, AnimSequence (wrapping existing C++ APIs)

C++ extension libraries:
- MovieSceneSectionEasingExtensions (SequencerScripting module): Get/SetEaseInDuration, Get/SetEaseOutDuration as ScriptMethod extensions on UMovieSceneSection
- OutlinerScriptingObject additions (SequencerCore): SetExpansion, IsExpanded on the existing outliner scripting object
- New MovieSceneAnimMixerScripting module with 4 extension libraries: track extensions (layers, animations, transitions), layer extensions (name, index, sections, empty), transition extensions (info, from/to sections, change type), decoration container extensions (generic discover/add/remove/find)

Other changes:
- Fixed pre-existing BuildDecorationMenu 5-arg to 4-arg signature mismatch in MovieSceneAnimationMixerTrackEditor.cpp
- Added MovieSceneAnimMixerScripting module to MovieSceneAnimMixer.uplugin
- Added LevelSequenceEditor, SequencerScripting, MovieSceneAnimMixer plugin dependencies to AnimationAssistantToolset.uplugin
- Updated init_unreal.py and __init__.py for SequencerTools registration and hot-reload

Tests (test_sequencer_toolset.py):
- 124 tests across 2 test classes (SequencerToolsTestCase and SequencerControlRigTestCase)

#Synced-CL 52492758

diff --git a/Engine/Plugins/Experimental/MovieSceneAnimMixer/MovieSceneAnimMixer.uplugin b/Engine/Plugins/Experimental/MovieSceneAnimMixer/MovieSceneAnimMixer.uplugin
--- a/Engine/Plugins/Experimental/MovieSceneAnimMixer/MovieSceneAnimMixer.uplugin
+++ b/Engine/Plugins/Experimental/MovieSceneAnimMixer/MovieSceneAnimMixer.uplugin
@@ -46,11 +46,6 @@
 			"Name" : "MovieSceneAnimMixerEditor",
 			"Type": "UncookedOnly",
 			"LoadingPhase": "Default"
-		},
-		{
-			"Name" : "MovieSceneAnimMixerScripting",
-			"Type" : "UncookedOnly",
-			"LoadingPhase" : "Default"
 		}
 	]
 }
diff --git a/Engine/Plugins/Experimental/MovieSceneAnimMixer/Source/MovieSceneAnimMixerEditor/Private/MovieSceneAnimationMixerTrackEditor.cpp b/Engine/Plugins/Experimental/MovieSceneAnimMixer/Source/MovieSceneAnimMixerEditor/Private/MovieSceneAnimationMixerTrackEditor.cpp
--- a/Engine/Plugins/Experimental/MovieSceneAnimMixer/Source/MovieSceneAnimMixerEditor/Private/MovieSceneAnimationMixerTrackEditor.cpp
+++ b/Engine/Plugins/Experimental/MovieSceneAnimMixer/Source/MovieSceneAnimMixerEditor/Private/MovieSceneAnimationMixerTrackEditor.cpp
@@ -1180,12 +1180,14 @@ TSharedRef<SWidget> FAnimationMixerTrackEditor::BuildAddSectionSubMenu(TWeakView
 		UMovieSceneAnimationMixerLayer* Layer = MixerTrack->GetLayer(LayerExtension->GetRowIndex());
 		if (Layer)
 		{
-			SequencerHelpers::BuildDecorationMenu(MenuBuilder, Layer, ObjectBindingGuid, Sequencer);
+			SequencerHelpers::BuildDecorationMenu(MenuBuilder, Layer, ObjectBindingGuid, Sequencer,
+				LOCTEXT("LayerModifiers", "Layer Modifiers"));
 
 			if (Layer->HasChildTrack())
 			{
 				// Show decorations compatible with the child track (e.g., root motion settings)
-				SequencerHelpers::BuildDecorationMenu(MenuBuilder, Layer->GetChildTrack(), ObjectBindingGuid, Sequencer);
+				SequencerHelpers::BuildDecorationMenu(MenuBuilder, Layer->GetChildTrack(), ObjectBindingGuid, Sequencer,
+					LOCTEXT("TrackModifiers", "Track Modifiers"));
 			}
 			else
 			{
diff --git a/Engine/Plugins/Experimental/MovieSceneAnimMixer/Source/MovieSceneAnimMixerScripting/MovieSceneAnimMixerScripting.Build.cs b/Engine/Plugins/Experimental/MovieSceneAnimMixer/Source/MovieSceneAnimMixerScripting/MovieSceneAnimMixerScripting.Build.cs
deleted file mode 100644
diff --git a/Engine/Plugins/Experimental/MovieSceneAnimMixer/Source/MovieSceneAnimMixerScripting/Private/MovieSceneAnimMixerLayerExtensions.cpp b/Engine/Plugins/Experimental/MovieSceneAnimMixer/Source/MovieSceneAnimMixerScripting/Private/MovieSceneAnimMixerLayerExtensions.cpp
deleted file mode 100644
diff --git a/Engine/Plugins/Experimental/MovieSceneAnimMixer/Source/MovieSceneAnimMixerScripting/Private/MovieSceneAnimMixerScriptingModule.cpp b/Engine/Plugins/Experimental/MovieSceneAnimMixer/Source/MovieSceneAnimMixerScripting/Private/MovieSceneAnimMixerScriptingModule.cpp
deleted file mode 100644
diff --git a/Engine/Plugins/Experimental/MovieSceneAnimMixer/Source/MovieSceneAnimMixerScripting/Private/MovieSceneAnimMixerScriptingModule.h b/Engine/Plugins/Experimental/MovieSceneAnimMixer/Source/MovieSceneAnimMixerScripting/Private/MovieSceneAnimMixerScriptingModule.h
deleted file mode 100644
diff --git a/Engine/Plugins/Experimental/MovieSceneAnimMixer/Source/MovieSceneAnimMixerScripting/Private/MovieSceneAnimMixerTrackExtensions.cpp b/Engine/Plugins/Experimental/MovieSceneAnimMixer/Source/MovieSceneAnimMixerScripting/Private/MovieSceneAnimMixerTrackExtensions.cpp
deleted file mode 100644
diff --git a/Engine/Plugins/Experimental/MovieSceneAnimMixer/Source/MovieSceneAnimMixerScripting/Private/MovieSceneAnimMixerTransitionExtensions.cpp b/Engine/Plugins/Experimental/MovieSceneAnimMixer/Source/MovieSceneAnimMixerScripting/Private/MovieSceneAnimMixerTransitionExtensions.cpp
deleted file mode 100644
diff --git a/Engine/Plugins/Experimental/MovieSceneAnimMixer/Source/MovieSceneAnimMixerScripting/Private/MovieSceneDecorationContainerExtensions.cpp b/Engine/Plugins/Experimental/MovieSceneAnimMixer/Source/MovieSceneAnimMixerScripting/Private/MovieSceneDecorationContainerExtensions.cpp
deleted file mode 100644
diff --git a/Engine/Plugins/Experimental/MovieSceneAnimMixer/Source/MovieSceneAnimMixerScripting/Public/MovieSceneAnimMixerLayerExtensions.h b/Engine/Plugins/Experimental/MovieSceneAnimMixer/Source/MovieSceneAnimMixerScripting/Public/MovieSceneAnimMixerLayerExtensions.h
deleted file mode 100644
diff --git a/Engine/Plugins/Experimental/MovieSceneAnimMixer/Source/MovieSceneAnimMixerScripting/Public/MovieSceneAnimMixerTrackExtensions.h b/Engine/Plugins/Experimental/MovieSceneAnimMixer/Source/MovieSceneAnimMixerScripting/Public/MovieSceneAnimMixerTrackExtensions.h
deleted file mode 100644
diff --git a/Engine/Plugins/Experimental/MovieSceneAnimMixer/Source/MovieSceneAnimMixerScripting/Public/MovieSceneAnimMixerTransitionExtensions.h b/Engine/Plugins/Experimental/MovieSceneAnimMixer/Source/MovieSceneAnimMixerScripting/Public/MovieSceneAnimMixerTransitionExtensions.h
deleted file mode 100644
diff --git a/Engine/Plugins/Experimental/MovieSceneAnimMixer/Source/MovieSceneAnimMixerScripting/Public/MovieSceneDecorationContainerExtensions.h b/Engine/Plugins/Experimental/MovieSceneAnimMixer/Source/MovieSceneAnimMixerScripting/Public/MovieSceneDecorationContainerExtensions.h
deleted file mode 100644
diff --git a/Engine/Plugins/Experimental/Toolsets/AnimationAssistantToolset/AnimationAssistantToolset.uplugin b/Engine/Plugins/Experimental/Toolsets/AnimationAssistantToolset/AnimationAssistantToolset.uplugin
--- a/Engine/Plugins/Experimental/Toolsets/AnimationAssistantToolset/AnimationAssistantToolset.uplugin
+++ b/Engine/Plugins/Experimental/Toolsets/AnimationAssistantToolset/AnimationAssistantToolset.uplugin
@@ -31,18 +31,6 @@
 		{
 			"Name": "ToolsetRegistry",
 			"Enabled": true
-		},
-		{
-			"Name": "LevelSequenceEditor",
-			"Enabled": true
-		},
-		{
-			"Name": "SequencerScripting",
-			"Enabled": true
-		},
-		{
-			"Name": "MovieSceneAnimMixer",
-			"Enabled": true
 		}
 	]
 }
\ No newline at end of file
diff --git a/Engine/Plugins/MovieScene/SequencerScripting/Source/SequencerScripting/Private/ExtensionLibraries/MovieSceneSectionEasingExtensions.cpp b/Engine/Plugins/MovieScene/SequencerScripting/Source/SequencerScripting/Private/ExtensionLibraries/MovieSceneSectionEasingExtensions.cpp
deleted file mode 100644
diff --git a/Engine/Plugins/MovieScene/SequencerScripting/Source/SequencerScripting/Public/ExtensionLibraries/MovieSceneSectionEasingExtensions.h b/Engine/Plugins/MovieScene/SequencerScripting/Source/SequencerScripting/Public/ExtensionLibraries/MovieSceneSectionEasingExtensions.h
deleted file mode 100644
diff --git a/Engine/Source/Editor/SequencerCore/Private/Scripting/OutlinerScriptingObject.cpp b/Engine/Source/Editor/SequencerCore/Private/Scripting/OutlinerScriptingObject.cpp
--- a/Engine/Source/Editor/SequencerCore/Private/Scripting/OutlinerScriptingObject.cpp
+++ b/Engine/Source/Editor/SequencerCore/Private/Scripting/OutlinerScriptingObject.cpp
@@ -343,32 +343,4 @@ void USequencerOutlinerScriptingObject::SetPinnedNodes(const TArray<FSequencerVi
 	}
 }
 
-void USequencerOutlinerScriptingObject::SetExpansion(const TArray<FSequencerViewModelScriptingStruct>& InNodes, bool bExpanded)
-{
-	using namespace UE::Sequencer;
-
-	for (const FSequencerViewModelScriptingStruct& Item : InNodes)
-	{
-		TViewModelPtr<IOutlinerExtension> OutlinerItem = Item.WeakViewModel.ImplicitPin();
-		if (OutlinerItem)
-		{
-			OutlinerItem->SetExpansion(bExpanded);
-		}
-	}
-}
-
-bool USequencerOutlinerScriptingObject::IsExpanded(FSequencerViewModelScriptingStruct InNode) const
-{
-	using namespace UE::Sequencer;
-
-	TViewModelPtr<IOutlinerExtension> OutlinerItem = InNode.WeakViewModel.ImplicitPin();
-	if (OutlinerItem)
-	{
-		return OutlinerItem->IsExpanded();
-	}
-
-	FFrame::KismetExecutionMessage(*LOCTEXT("ViewModelNotValid", "View Model is no longer valid.").ToString(), ELogVerbosity::Error);
-	return false;
-}
-
 #undef LOCTEXT_NAMESPACE
diff --git a/Engine/Source/Editor/SequencerCore/Public/Scripting/OutlinerScriptingObject.h b/Engine/Source/Editor/SequencerCore/Public/Scripting/OutlinerScriptingObject.h
--- a/Engine/Source/Editor/SequencerCore/Public/Scripting/OutlinerScriptingObject.h
+++ b/Engine/Source/Editor/SequencerCore/Public/Scripting/OutlinerScriptingObject.h
@@ -74,13 +74,6 @@ public:
 	UFUNCTION(BlueprintCallable, Category = "Sequencer Editor")
 	void SetPinnedNodes(const TArray<FSequencerViewModelScriptingStruct>& InNodes, bool bInPinned);
 
-	// Expansion
-	UFUNCTION(BlueprintCallable, Category = "Sequencer Editor")
-	void SetExpansion(const TArray<FSequencerViewModelScriptingStruct>& InNodes, bool bExpanded);
-
-	UFUNCTION(BlueprintPure, Category = "Sequencer Editor")
-	bool IsExpanded(FSequencerViewModelScriptingStruct InNode) const;
-
 protected:
 
 	UE::Sequencer::TWeakViewModelPtr<UE::Sequencer::FOutlinerViewModel> WeakOutliner;
