Sequencer: Fix accumulated offset discontinuities during transitions and at loop boundaries

The accumulated offset cache is global across all sections, but lookups
for a specific section were picking up entries from other sections whose
start frames fell within the current section's range. This caused root
motion teleports at transition boundaries and loop points.

Adds OwnerSection to cache entries and a section-filtered lookup method
(GetAccumulatedOffsetForSection) that only considers entries owned by the
requesting section or from before its start. Used both at runtime in
BuildSpaceConversionTask and during cache building so loop entries chain
from the correct base offset.

#Synced-CL 52530536

diff --git a/Engine/Plugins/Experimental/MovieSceneAnimMixer/Source/MovieSceneAnimMixer/Private/MovieSceneRootMotionSection.cpp b/Engine/Plugins/Experimental/MovieSceneAnimMixer/Source/MovieSceneAnimMixer/Private/MovieSceneRootMotionSection.cpp
--- a/Engine/Plugins/Experimental/MovieSceneAnimMixer/Source/MovieSceneAnimMixer/Private/MovieSceneRootMotionSection.cpp
+++ b/Engine/Plugins/Experimental/MovieSceneAnimMixer/Source/MovieSceneAnimMixer/Private/MovieSceneRootMotionSection.cpp
@@ -120,15 +120,36 @@ FTransform UMovieSceneRootMotionSection::GetAccumulatedOffsetAt(FFrameTime Time)
 	return Result;
 }
 
+FTransform UMovieSceneRootMotionSection::GetAccumulatedOffsetForSection(FFrameTime Time, UMovieSceneSection* Section) const
+{
+	return FindAccumulatedOffsetForSection(Time, Section, false);
+}
+
 FTransform UMovieSceneRootMotionSection::GetAnimSpaceAccumulatedOffsetAt(FFrameTime Time) const
+{
+	return FindAccumulatedOffsetForSection(Time, nullptr, true);
+}
+
+FTransform UMovieSceneRootMotionSection::GetAnimSpaceAccumulatedOffsetForSection(FFrameTime Time, UMovieSceneSection* Section) const
+{
+	return FindAccumulatedOffsetForSection(Time, Section, true);
+}
+
+FTransform UMovieSceneRootMotionSection::FindAccumulatedOffsetForSection(FFrameTime Time, UMovieSceneSection* Section, bool bAnimSpace) const
 {
 	FTransform Result = FTransform::Identity;
 
+	const FFrameTime SectionStart = Section && Section->HasStartFrame()
+		? FFrameTime(Section->GetInclusiveStartFrame()) : FFrameTime();
+
 	for (const FRootMotionAccumulatedOffset& Offset : AccumulatedOffsetCache)
 	{
 		if (Offset.StartFrame <= Time)
 		{
-			Result = Offset.AnimSpaceAccumulatedOffset;
+			if (!Section || Offset.StartFrame <= SectionStart || Offset.OwnerSection == Section || !Offset.OwnerSection)
+			{
+				Result = bAnimSpace ? Offset.AnimSpaceAccumulatedOffset : Offset.AccumulatedOffset;
+			}
 		}
 		else
 		{
@@ -308,6 +329,7 @@ void UMovieSceneRootMotionSection::RebuildAccumulatedOffsetCache(
 			InitialEntry.AccumulatedOffset = Settings->InitialRootTransform;
 			InitialEntry.AnimSpaceAccumulatedOffset = Settings->InitialRootTransform;
 			InitialEntry.Weight = 1.0f;
+			InitialEntry.OwnerSection = Section;
 			AccumulatedOffsetCache.Add(InitialEntry);
 		}
 	}
@@ -362,14 +384,15 @@ void UMovieSceneRootMotionSection::RebuildAccumulatedOffsetCache(
 			FRootMotionAccumulatedOffset CacheEntry;
 			CacheEntry.StartFrame = Pending.EntryStartFrame;
 			CacheEntry.Weight = 1.0f;
+			CacheEntry.OwnerSection = Pending.Section;
 
 			if (Pending.Section->HasStartFrame())
 			{
-				// Use the entry's own start frame for the accumulated offset lookup.
-				// For loop boundary entries within the same section, this chains each
-				// loop iteration on top of the previous one's accumulated offset.
-				CacheEntry.AccumulatedOffset = Delta * GetAccumulatedOffsetAt(Pending.EntryStartFrame);
-				CacheEntry.AnimSpaceAccumulatedOffset = Delta * GetAnimSpaceAccumulatedOffsetAt(Pending.EntryStartFrame);
+				// Use the section-filtered lookup so loop entries chain from
+				// this section's own prior entries, not from entries of other
+				// sections that happen to fall within this section's range.
+				CacheEntry.AccumulatedOffset = Delta * GetAccumulatedOffsetForSection(Pending.EntryStartFrame, Pending.Section);
+				CacheEntry.AnimSpaceAccumulatedOffset = Delta * GetAnimSpaceAccumulatedOffsetForSection(Pending.EntryStartFrame, Pending.Section);
 			}
 			else
 			{
@@ -393,6 +416,7 @@ void UMovieSceneRootMotionSection::RebuildAccumulatedOffsetCache(
 			CacheEntry.AccumulatedOffset = Result.RootMotionTransform;
 			CacheEntry.AnimSpaceAccumulatedOffset = Result.AnimationSpaceRootMotion;
 			CacheEntry.Weight = 1.0f;
+			CacheEntry.OwnerSection = Pending.Section;
 			AccumulatedOffsetCache.Add(CacheEntry);
 		}
 	}
diff --git a/Engine/Plugins/Experimental/MovieSceneAnimMixer/Source/MovieSceneAnimMixer/Private/Systems/MovieSceneAnimMixerSystem.cpp b/Engine/Plugins/Experimental/MovieSceneAnimMixer/Source/MovieSceneAnimMixer/Private/Systems/MovieSceneAnimMixerSystem.cpp
--- a/Engine/Plugins/Experimental/MovieSceneAnimMixer/Source/MovieSceneAnimMixer/Private/Systems/MovieSceneAnimMixerSystem.cpp
+++ b/Engine/Plugins/Experimental/MovieSceneAnimMixer/Source/MovieSceneAnimMixer/Private/Systems/MovieSceneAnimMixerSystem.cpp
@@ -168,23 +168,23 @@ static TOptional<FAnimNextConvertRootMotionToWorldSpaceTask> BuildSpaceConversio
 				{
 					if (auto* RootMotionSection = TargetDecoration->GetRootMotionSection())
 					{
-						// Use the current evaluation time so loop boundary cache entries
-						// within a section are picked up correctly.
-						FFrameTime AccumLookupTime(OwnerSection->HasStartFrame()
+						FFrameTime SectionStart(OwnerSection->HasStartFrame()
 							? OwnerSection->GetInclusiveStartFrame() : FFrameNumber(0));
+
+						FFrameTime CurrentTime = SectionStart;
 						if (Context.Linker && MixerEntry.InstanceHandle.IsValid())
 						{
-							const FMovieSceneContext& MovieContext = Context.Linker->GetInstanceRegistry()->GetContext(MixerEntry.InstanceHandle);
-							AccumLookupTime = MovieContext.GetTime();
+							CurrentTime = Context.Linker->GetInstanceRegistry()->GetContext(MixerEntry.InstanceHandle).GetTime();
 						}
+
 						bUseAccumulatedOffset = !RootMotionSection->AccumulatedOffsetCache.IsEmpty()
-							&& RootMotionSection->AccumulatedOffsetCache[0].StartFrame <= AccumLookupTime;
+							&& RootMotionSection->AccumulatedOffsetCache[0].StartFrame <= SectionStart;
 
 						if (bUseAccumulatedOffset)
 						{
 							Conversion |= ESpaceConversions::AccumulatedOffset;
 							ResolvedAccumulatedTransform =
-								RootMotionSection->GetAccumulatedOffsetAt(AccumLookupTime);
+								RootMotionSection->GetAccumulatedOffsetForSection(CurrentTime, OwnerSection);
 						}
 					}
 				}
diff --git a/Engine/Plugins/Experimental/MovieSceneAnimMixer/Source/MovieSceneAnimMixer/Public/MovieSceneRootMotionSection.h b/Engine/Plugins/Experimental/MovieSceneAnimMixer/Source/MovieSceneAnimMixer/Public/MovieSceneRootMotionSection.h
--- a/Engine/Plugins/Experimental/MovieSceneAnimMixer/Source/MovieSceneAnimMixer/Public/MovieSceneRootMotionSection.h
+++ b/Engine/Plugins/Experimental/MovieSceneAnimMixer/Source/MovieSceneAnimMixer/Public/MovieSceneRootMotionSection.h
@@ -60,6 +60,10 @@ struct FRootMotionAccumulatedOffset
 	/** Weight of this offset for blending */
 	UPROPERTY()
 	float Weight = 1.0f;
+
+	/** The section that generated this cache entry */
+	UPROPERTY()
+	TObjectPtr<UMovieSceneSection> OwnerSection = nullptr;
 };
 
 /**
@@ -93,9 +97,20 @@ public:
 	/** Get the accumulated offset at a given time (world space) */
 	MOVIESCENEANIMMIXER_API FTransform GetAccumulatedOffsetAt(FFrameTime Time) const;
 
+	/** Get the accumulated offset at a given time, only considering entries
+	 *  owned by the given section or from before the section's start. */
+	MOVIESCENEANIMMIXER_API FTransform GetAccumulatedOffsetForSection(FFrameTime Time, UMovieSceneSection* Section) const;
+
 	/** Get the accumulated offset at a given time in animation space (pre-world conversion) */
 	MOVIESCENEANIMMIXER_API FTransform GetAnimSpaceAccumulatedOffsetAt(FFrameTime Time) const;
 
+	/** Get the anim-space accumulated offset, section-filtered. */
+	MOVIESCENEANIMMIXER_API FTransform GetAnimSpaceAccumulatedOffsetForSection(FFrameTime Time, UMovieSceneSection* Section) const;
+
+private:
+	FTransform FindAccumulatedOffsetForSection(FFrameTime Time, UMovieSceneSection* Section, bool bAnimSpace) const;
+
+public:
 	/** Rebuild the accumulated offset cache using isolated bake evaluations.*/
 	MOVIESCENEANIMMIXER_API void RebuildAccumulatedOffsetCache(
 		class UMovieSceneEntitySystemLinker* Linker,
