Sequencer: Animation Bus Track and Bus Target for Animation Mixer

Adds cross-sequence animation routing via named buses.
A new Bus Target type writes mixer results to named bus storage.
A new Bus Section type reads from bus storage into mixer layers.
Includes dependency-ordered evaluation, cycle detection, and editor UI.

#Synced-CL 52530536

diff --git a/Engine/Plugins/Experimental/MovieSceneAnimMixer/Source/MovieSceneAnimMixer/Private/AnimMixerBusUtils.cpp b/Engine/Plugins/Experimental/MovieSceneAnimMixer/Source/MovieSceneAnimMixer/Private/AnimMixerBusUtils.cpp
new file mode 100644
--- /dev/null
+++ b/Engine/Plugins/Experimental/MovieSceneAnimMixer/Source/MovieSceneAnimMixer/Private/AnimMixerBusUtils.cpp
@@ -0,0 +1,457 @@
+// Copyright Epic Games, Inc. All Rights Reserved.
+
+#include "AnimMixerBusUtils.h"
+
+#include "MovieSceneAnimBusSection.h"
+#include "MovieSceneAnimationMixerTrack.h"
+#include "Systems/MovieSceneAnimMixerSystem.h"
+#include "MovieScene.h"
+#include "MovieSceneSequence.h"
+#include "Sections/MovieSceneSubSection.h"
+#include "Tracks/MovieSceneSubTrack.h"
+
+// ---------------------------------------------------------------------------
+// FBusDependencyGraph
+// ---------------------------------------------------------------------------
+
+TArray<int32> FBusDependencyGraph::TopologicalSort()
+{
+	TMap<FName, int32> WriterByBus;
+	for (int32 i = 0; i < Nodes.Num(); ++i)
+	{
+		if (!Nodes[i].WritesToBus.IsNone())
+		{
+			WriterByBus.Add(Nodes[i].WritesToBus, i);
+		}
+		Nodes[i].InDegree = 0;
+		Nodes[i].Dependents.Reset();
+	}
+
+	for (int32 ReaderIdx = 0; ReaderIdx < Nodes.Num(); ++ReaderIdx)
+	{
+		for (const FName& ReadBus : Nodes[ReaderIdx].ReadsFromBuses)
+		{
+			if (ReadBus == Nodes[ReaderIdx].WritesToBus)
+			{
+				continue;
+			}
+			if (int32* WriterIdx = WriterByBus.Find(ReadBus))
+			{
+				if (*WriterIdx != ReaderIdx)
+				{
+					Nodes[*WriterIdx].Dependents.AddUnique(ReaderIdx);
+					Nodes[ReaderIdx].InDegree++;
+				}
+			}
+		}
+	}
+
+	TArray<int32> Queue;
+	for (int32 i = 0; i < Nodes.Num(); ++i)
+	{
+		if (Nodes[i].InDegree == 0)
+		{
+			Queue.Add(i);
+		}
+	}
+
+	TArray<int32> Result;
+	Result.Reserve(Nodes.Num());
+	int32 QueueHead = 0;
+	while (QueueHead < Queue.Num())
+	{
+		int32 Current = Queue[QueueHead++];
+		Result.Add(Current);
+		for (int32 DepIdx : Nodes[Current].Dependents)
+		{
+			if (--Nodes[DepIdx].InDegree == 0)
+			{
+				Queue.Add(DepIdx);
+			}
+		}
+	}
+
+	return Result;
+}
+
+bool FBusDependencyGraph::IsReachable(FName StartBus, FName TargetBus) const
+{
+	TMap<FName, TArrayView<const FName>> AdjMap;
+	for (const FNode& Node : Nodes)
+	{
+		if (!Node.WritesToBus.IsNone() && !Node.ReadsFromBuses.IsEmpty())
+		{
+			AdjMap.Add(Node.WritesToBus, Node.ReadsFromBuses);
+		}
+	}
+
+	TSet<FName> Visited;
+	TArray<FName> Stack;
+	Stack.Add(StartBus);
+	while (!Stack.IsEmpty())
+	{
+		FName Current = Stack.Pop();
+		if (Current == TargetBus)
+		{
+			return true;
+		}
+		if (Visited.Contains(Current))
+		{
+			continue;
+		}
+		Visited.Add(Current);
+
+		if (const TArrayView<const FName>* Reads = AdjMap.Find(Current))
+		{
+			for (const FName& Name : *Reads)
+			{
+				Stack.Add(Name);
+			}
+		}
+	}
+	return false;
+}
+
+// ---------------------------------------------------------------------------
+// Helpers
+// ---------------------------------------------------------------------------
+
+namespace
+{
+
+TArray<FName> GatherBusReads(UMovieSceneAnimationMixerTrack* Track)
+{
+	TArray<FName> Reads;
+	for (UMovieSceneSection* Section : Track->GetAllSections())
+	{
+		if (UMovieSceneAnimBusSection* BusSection = Cast<UMovieSceneAnimBusSection>(Section))
+		{
+			if (!BusSection->BusName.IsNone())
+			{
+				Reads.AddUnique(BusSection->BusName);
+			}
+		}
+	}
+	return Reads;
+}
+
+FBusDependencyGraph BuildGraphFromTracks(TArrayView<UMovieSceneAnimationMixerTrack* const> MixerTracks)
+{
+	FBusDependencyGraph Graph;
+	for (UMovieSceneAnimationMixerTrack* Track : MixerTracks)
+	{
+		const FMovieSceneAnimBusTarget* BusTarget = Track->MixedAnimationTarget.GetPtr<FMovieSceneAnimBusTarget>();
+		if (!BusTarget || BusTarget->BusName.IsNone())
+		{
+			continue;
+		}
+
+		FBusDependencyGraph::FNode& Node = Graph.Nodes.AddDefaulted_GetRef();
+		Node.WritesToBus = BusTarget->BusName;
+		Node.ReadsFromBuses = GatherBusReads(Track);
+	}
+	return Graph;
+}
+
+void GatherMixerTracksRecursive(UMovieSceneSequence* Sequence, TArray<UMovieSceneAnimationMixerTrack*>& OutTracks, TSet<UMovieSceneSequence*>& Visited)
+{
+	if (!Sequence || Visited.Contains(Sequence))
+	{
+		return;
+	}
+	Visited.Add(Sequence);
+
+	UMovieScene* MovieScene = Sequence->GetMovieScene();
+	if (!MovieScene)
+	{
+		return;
+	}
+
+	const UMovieScene* ConstMovieScene = MovieScene;
+	for (const FMovieSceneBinding& Binding : ConstMovieScene->GetBindings())
+	{
+		for (UMovieSceneTrack* Track : Binding.GetTracks())
+		{
+			if (UMovieSceneAnimationMixerTrack* MixerTrack = Cast<UMovieSceneAnimationMixerTrack>(Track))
+			{
+				OutTracks.Add(MixerTrack);
+			}
+		}
+	}
+
+	for (UMovieSceneTrack* MasterTrack : MovieScene->GetTracks())
+	{
+		if (UMovieSceneSubTrack* SubTrack = Cast<UMovieSceneSubTrack>(MasterTrack))
+		{
+			for (UMovieSceneSection* Section : SubTrack->GetAllSections())
+			{
+				if (UMovieSceneSubSection* SubSection = Cast<UMovieSceneSubSection>(Section))
+				{
+					if (UMovieSceneSequence* SubSequence = SubSection->GetSequence())
+					{
+						GatherMixerTracksRecursive(SubSequence, OutTracks, Visited);
+					}
+				}
+			}
+		}
+	}
+}
+
+} // anonymous namespace
+
+// ---------------------------------------------------------------------------
+// FAnimMixerBusUtils
+// ---------------------------------------------------------------------------
+
+TArray<FName> FAnimMixerBusUtils::GatherBusNamesFromSequence(UMovieSceneSequence* RootSequence)
+{
+	TSet<FName> BusNames;
+
+	TArray<UMovieSceneAnimationMixerTrack*> Tracks;
+	TSet<UMovieSceneSequence*> Visited;
+	GatherMixerTracksRecursive(RootSequence, Tracks, Visited);
+
+	for (UMovieSceneAnimationMixerTrack* Track : Tracks)
+	{
+		if (const FMovieSceneAnimBusTarget* BusTarget = Track->MixedAnimationTarget.GetPtr<FMovieSceneAnimBusTarget>())
+		{
+			if (!BusTarget->BusName.IsNone())
+			{
+				BusNames.Add(BusTarget->BusName);
+			}
+		}
+
+		for (UMovieSceneSection* Section : Track->GetAllSections())
+		{
+			if (UMovieSceneAnimBusSection* BusSection = Cast<UMovieSceneAnimBusSection>(Section))
+			{
+				if (!BusSection->BusName.IsNone())
+				{
+					BusNames.Add(BusSection->BusName);
+				}
+			}
+		}
+	}
+
+	TArray<FName> Result = BusNames.Array();
+	Result.Sort(FNameLexicalLess());
+	return Result;
+}
+
+TArray<FName> FAnimMixerBusUtils::ComputeBusEvaluationOrder(TArrayView<UMovieSceneAnimationMixerTrack* const> MixerTracks)
+{
+	FBusDependencyGraph Graph = BuildGraphFromTracks(MixerTracks);
+	if (Graph.Nodes.IsEmpty())
+	{
+		return {};
+	}
+
+	// Mark self-referencing nodes (TopologicalSort skips self-edges so they
+	// wouldn't be caught by the cycle check below)
+	TSet<FName> SelfReferencingBuses;
+	for (const FBusDependencyGraph::FNode& Node : Graph.Nodes)
+	{
+		if (Node.ReadsFromBuses.Contains(Node.WritesToBus))
+		{
+			SelfReferencingBuses.Add(Node.WritesToBus);
+			ensureMsgf(false, TEXT("Bus '%s' reads from itself."), *Node.WritesToBus.ToString());
+		}
+	}
+
+	TArray<int32> Sorted = Graph.TopologicalSort();
+
+	if (Sorted.Num() < Graph.Nodes.Num())
+	{
+		for (const FBusDependencyGraph::FNode& Node : Graph.Nodes)
+		{
+			if (Node.InDegree > 0)
+			{
+				ensureMsgf(false, TEXT("Bus dependency cycle detected involving bus '%s'."), *Node.WritesToBus.ToString());
+			}
+		}
+	}
+
+	TArray<FName> Result;
+	Result.Reserve(Sorted.Num());
+	for (int32 Idx : Sorted)
+	{
+		FName BusName = Graph.Nodes[Idx].WritesToBus;
+		if (!SelfReferencingBuses.Contains(BusName))
+		{
+			Result.Add(BusName);
+		}
+	}
+	return Result;
+}
+
+FAnimMixerBusValidationResult FAnimMixerBusUtils::ValidateBusTopology(TArrayView<UMovieSceneAnimationMixerTrack* const> MixerTracks)
+{
+	FAnimMixerBusValidationResult Result;
+
+	TArray<FName> Writers;
+	TArray<FName> Readers;
+
+	for (UMovieSceneAnimationMixerTrack* Track : MixerTracks)
+	{
+		if (const FMovieSceneAnimBusTarget* BusTarget = Track->MixedAnimationTarget.GetPtr<FMovieSceneAnimBusTarget>())
+		{
+			if (!BusTarget->BusName.IsNone())
+			{
+				Writers.Add(BusTarget->BusName);
+			}
+		}
+
+		for (UMovieSceneSection* Section : Track->GetAllSections())
+		{
+			if (UMovieSceneAnimBusSection* BusSection = Cast<UMovieSceneAnimBusSection>(Section))
+			{
+				if (!BusSection->BusName.IsNone())
+				{
+					Readers.Add(BusSection->BusName);
+				}
+			}
+		}
+	}
+
+	// Duplicate writer detection
+	TMap<FName, int32> WriterCounts;
+	for (const FName& Name : Writers)
+	{
+		WriterCounts.FindOrAdd(Name)++;
+	}
+	for (const auto& Pair : WriterCounts)
+	{
+		if (Pair.Value > 1)
+		{
+			Result.Warnings.Add(FString::Printf(
+				TEXT("Multiple mixers write to bus '%s'. Only one will take effect per frame."),
+				*Pair.Key.ToString()));
+		}
+	}
+
+	// Orphan bus section detection
+	TSet<FName> WrittenBuses(Writers);
+	for (const FName& Name : Readers)
+	{
+		if (!WrittenBuses.Contains(Name))
+		{
+			Result.Warnings.Add(FString::Printf(
+				TEXT("Bus section references bus '%s' but no mixer targets that bus."),
+				*Name.ToString()));
+		}
+	}
+
+	// Orphan target detection
+	TSet<FName> ReadBuses(Readers);
+	for (const FName& Name : Writers)
+	{
+		if (!ReadBuses.Contains(Name))
+		{
+			Result.Warnings.Add(FString::Printf(
+				TEXT("Mixer targets bus '%s' but no bus section reads from it."),
+				*Name.ToString()));
+		}
+	}
+
+	// Self-reference + cycle detection
+	FBusDependencyGraph Graph = BuildGraphFromTracks(MixerTracks);
+
+	for (const FBusDependencyGraph::FNode& Node : Graph.Nodes)
+	{
+		if (Node.ReadsFromBuses.Contains(Node.WritesToBus))
+		{
+			Result.Errors.Add(FString::Printf(
+				TEXT("Mixer reads from bus '%s' which it also writes to."),
+				*Node.WritesToBus.ToString()));
+		}
+	}
+
+	TArray<int32> Sorted = Graph.TopologicalSort();
+	if (Sorted.Num() < Graph.Nodes.Num())
+	{
+		TArray<FString> CycleBuses;
+		for (int32 i = 0; i < Graph.Nodes.Num(); ++i)
+		{
+			if (Graph.Nodes[i].InDegree > 0)
+			{
+				CycleBuses.Add(Graph.Nodes[i].WritesToBus.ToString());
+			}
+		}
+		Result.Errors.Add(FString::Printf(
+			TEXT("Bus dependency cycle detected involving buses: %s."),
+			*FString::Join(CycleBuses, TEXT(", "))));
+	}
+
+	return Result;
+}
+
+bool FAnimMixerBusUtils::WouldBusSectionCreateCycle(FName BusName, UMovieSceneAnimationMixerTrack* Track,
+	TArrayView<UMovieSceneAnimationMixerTrack* const> MixerTracks)
+{
+	if (!Track || BusName.IsNone())
+	{
+		return false;
+	}
+
+	const FMovieSceneAnimBusTarget* BusTarget = Track->MixedAnimationTarget.GetPtr<FMovieSceneAnimBusTarget>();
+	if (!BusTarget || BusTarget->BusName.IsNone())
+	{
+		return false;
+	}
+
+	if (BusTarget->BusName == BusName)
+	{
+		return true;
+	}
+
+	FBusDependencyGraph Graph = BuildGraphFromTracks(MixerTracks);
+
+	// Add the hypothetical read
+	for (FBusDependencyGraph::FNode& Node : Graph.Nodes)
+	{
+		if (Node.WritesToBus == BusTarget->BusName)
+		{
+			Node.ReadsFromBuses.AddUnique(BusName);
+			break;
+		}
+	}
+
+	return Graph.IsReachable(BusName, BusTarget->BusName);
+}
+
+bool FAnimMixerBusUtils::WouldBusTargetCreateCycle(FName BusName, UMovieSceneAnimationMixerTrack* Track,
+	TArrayView<UMovieSceneAnimationMixerTrack* const> MixerTracks)
+{
+	if (!Track || BusName.IsNone())
+	{
+		return false;
+	}
+
+	TArray<FName> TrackReads = GatherBusReads(Track);
+	if (TrackReads.IsEmpty())
+	{
+		return false;
+	}
+
+	if (TrackReads.Contains(BusName))
+	{
+		return true;
+	}
+
+	FBusDependencyGraph Graph = BuildGraphFromTracks(MixerTracks);
+
+	// Add the hypothetical node
+	FBusDependencyGraph::FNode& HypNode = Graph.Nodes.AddDefaulted_GetRef();
+	HypNode.WritesToBus = BusName;
+	HypNode.ReadsFromBuses = TrackReads;
+
+	for (const FName& ReadBus : TrackReads)
+	{
+		if (Graph.IsReachable(ReadBus, BusName))
+		{
+			return true;
+		}
+	}
+
+	return false;
+}
diff --git a/Engine/Plugins/Experimental/MovieSceneAnimMixer/Source/MovieSceneAnimMixer/Private/AnimMixerComponentTypes.cpp b/Engine/Plugins/Experimental/MovieSceneAnimMixer/Source/MovieSceneAnimMixer/Private/AnimMixerComponentTypes.cpp
--- a/Engine/Plugins/Experimental/MovieSceneAnimMixer/Source/MovieSceneAnimMixer/Private/AnimMixerComponentTypes.cpp
+++ b/Engine/Plugins/Experimental/MovieSceneAnimMixer/Source/MovieSceneAnimMixer/Private/AnimMixerComponentTypes.cpp
@@ -55,7 +55,7 @@ namespace UE::MovieScene
 		ComponentRegistry->NewComponentType(&GapBehavior, TEXT("Gap Behavior"));
 		ComponentRegistry->NewComponentType(&EntityOwner, TEXT("Entity Owner"));
 		ComponentRegistry->NewComponentType(&LayerBlendTask, TEXT("Layer Blend Task"));
-
+		ComponentRegistry->NewComponentType(&BusName, TEXT("Bus Name"));
 
 		Tags.Transition = ComponentRegistry->NewTag(TEXT("Transition"));
 		Tags.SkipPoseWeight = ComponentRegistry->NewTag(TEXT("Skip Pose Weight"));
@@ -76,6 +76,7 @@ namespace UE::MovieScene
 		ComponentRegistry->Factories.DuplicateChildComponent(BoneMatchTransform);
 
 		ComponentRegistry->Factories.DuplicateChildComponent(MirrorTable);
+		ComponentRegistry->Factories.DuplicateChildComponent(BusName);
 		
 		ComponentRegistry->Factories.DefineChildComponent(RootDestination, MixerRootMotion);
 		ComponentRegistry->Factories.DefineChildComponent(Task, MixerEntry);
diff --git a/Engine/Plugins/Experimental/MovieSceneAnimMixer/Source/MovieSceneAnimMixer/Private/MovieSceneAnimBusSection.cpp b/Engine/Plugins/Experimental/MovieSceneAnimMixer/Source/MovieSceneAnimMixer/Private/MovieSceneAnimBusSection.cpp
new file mode 100644
--- /dev/null
+++ b/Engine/Plugins/Experimental/MovieSceneAnimMixer/Source/MovieSceneAnimMixer/Private/MovieSceneAnimBusSection.cpp
@@ -0,0 +1,94 @@
+// Copyright Epic Games, Inc. All Rights Reserved.
+
+#include "MovieSceneAnimBusSection.h"
+
+#include "AnimMixerComponentTypes.h"
+#include "Channels/MovieSceneChannelProxy.h"
+#include "MovieSceneAnimationMixerTrack.h"
+#include "MovieSceneAnimationMixerLayer.h"
+#include "EntitySystem/BuiltInComponentTypes.h"
+#include "Systems/MovieSceneSkeletalAnimationSystem.h"
+
+#include UE_INLINE_GENERATED_CPP_BY_NAME(MovieSceneAnimBusSection)
+
+#define LOCTEXT_NAMESPACE "MovieSceneAnimBusSection"
+
+UMovieSceneAnimBusSection::UMovieSceneAnimBusSection()
+{
+	bSupportsInfiniteRange = true;
+	SetRange(TRange<FFrameNumber>::All());
+	BlendType = EMovieSceneBlendType::Absolute;
+	SetColorTint(BusSectionTintColor);
+
+	// Default weight of 1.0 (full contribution)
+	Weight.SetDefault(1.0f);
+}
+
+EMovieSceneChannelProxyType UMovieSceneAnimBusSection::CacheChannelProxy()
+{
+	FMovieSceneChannelProxyData Channels;
+
+#if WITH_EDITOR
+	static FMovieSceneChannelMetaData MetaData("Weight", LOCTEXT("WeightChannelName", "Weight"));
+	MetaData.bCanCollapseToTrack = false;
+
+	Channels.Add(Weight, MetaData, TMovieSceneExternalValue<float>());
+#else
+	Channels.Add(Weight);
+#endif
+
+	ChannelProxy = MakeShared<FMovieSceneChannelProxy>(MoveTemp(Channels));
+	return EMovieSceneChannelProxyType::Dynamic;
+}
+
+void UMovieSceneAnimBusSection::ImportEntityImpl(UMovieSceneEntitySystemLinker* EntityLinker, const FEntityImportParams& Params, FImportedEntity* OutImportedEntity)
+{
+	using namespace UE::MovieScene;
+
+	const FAnimMixerComponentTypes* AnimMixerComponentTypes = FAnimMixerComponentTypes::Get();
+	FBuiltInComponentTypes* BuiltInComponentTypes = FBuiltInComponentTypes::Get();
+
+	BuildDefaultComponents(EntityLinker, Params, OutImportedEntity);
+
+	if (BusName.IsNone())
+	{
+		return;
+	}
+
+	// Find the owning mixer track and layer for this section
+	UMovieSceneAnimationMixerTrack* MixerTrack = GetTypedOuter<UMovieSceneAnimationMixerTrack>();
+	UMovieSceneAnimationMixerLayer* Layer = nullptr;
+	if (MixerTrack)
+	{
+		Layer = MixerTrack->GetLayerForSection(this);
+	}
+
+	// Task and MixerEntry start as null; UMovieSceneAnimBusSystem populates the
+	// Task with a FMovieSceneAnimBusReadTask each frame. The mixer system then
+	// wraps it into an entry like any other section.
+	TSharedPtr<FAnimNextEvaluationTask> NullTask;
+	TSharedPtr<FMovieSceneAnimMixerEntry> NullEntry;
+
+	// Resolve target from the owning mixer track
+	TInstancedStruct<FMovieSceneMixedAnimationTarget> ResolvedTarget;
+	if (MixerTrack && MixerTrack->MixedAnimationTarget.IsValid())
+	{
+		ResolvedTarget = MixerTrack->MixedAnimationTarget;
+	}
+
+	OutImportedEntity->AddBuilder(
+		FEntityBuilder()
+		.Add(BuiltInComponentTypes->GenericObjectBinding, Params.GetObjectBindingID())
+		.Add(BuiltInComponentTypes->BoundObjectResolver, UMovieSceneSkeletalAnimationSystem::ResolveSkeletalMeshComponentBinding)
+		.Add(AnimMixerComponentTypes->Priority, GetRowIndex())
+		.Add(AnimMixerComponentTypes->Task, NullTask)
+		.Add(AnimMixerComponentTypes->MixerEntry, NullEntry)
+		.Add(AnimMixerComponentTypes->EntityOwner, FObjectKey(this))
+		.Add(AnimMixerComponentTypes->BusName, BusName)
+		.AddConditional(BuiltInComponentTypes->WeightChannel, &Weight, Weight.HasAnyData())
+		.AddConditional(AnimMixerComponentTypes->MixerLayer, Layer, Layer != nullptr)
+		.AddConditional(AnimMixerComponentTypes->Target, ResolvedTarget, ResolvedTarget.IsValid())
+		.AddMutualComponents());
+}
+
+#undef LOCTEXT_NAMESPACE
diff --git a/Engine/Plugins/Experimental/MovieSceneAnimMixer/Source/MovieSceneAnimMixer/Private/MovieSceneAnimMixerModule.cpp b/Engine/Plugins/Experimental/MovieSceneAnimMixer/Source/MovieSceneAnimMixer/Private/MovieSceneAnimMixerModule.cpp
--- a/Engine/Plugins/Experimental/MovieSceneAnimMixer/Source/MovieSceneAnimMixer/Private/MovieSceneAnimMixerModule.cpp
+++ b/Engine/Plugins/Experimental/MovieSceneAnimMixer/Source/MovieSceneAnimMixer/Private/MovieSceneAnimMixerModule.cpp
@@ -12,6 +12,7 @@
 #include "MovieSceneLayerWeightDecoration.h"
 #include "Modules/ModuleInterface.h"
 #include "Modules/ModuleManager.h"
+#include "MovieSceneAnimBusSection.h"
 #include "Sections/MovieSceneSkeletalAnimationSection.h"
 
 #if WITH_UNREAL_DEVELOPER_TOOLS
@@ -50,8 +51,9 @@ void FMovieSceneAnimMixerModule::StartupModule()
 	// Register Root Motion Target decoration as compatible with the Animation Mixer Track
 	IMovieSceneModule::Get().RegisterCompatibleDecoration(UMovieSceneAnimationMixerTrack::StaticClass(), UMovieSceneRootMotionTargetDecoration::StaticClass());
 
-	// Register Root Motion Settings decoration as compatible with Skeletal Animation Sections
+	// Register Root Motion Settings decoration as compatible with Skeletal Animation Sections and Bus Sections
 	IMovieSceneModule::Get().RegisterCompatibleDecoration(UMovieSceneSkeletalAnimationSection::StaticClass(), UMovieSceneRootMotionSettingsDecoration::StaticClass());
+	IMovieSceneModule::Get().RegisterCompatibleDecoration(UMovieSceneAnimBusSection::StaticClass(), UMovieSceneRootMotionSettingsDecoration::StaticClass());
 
 	IMovieSceneModule::Get().RegisterCompatibleDecoration(UMovieSceneAnimationMixerLayer::StaticClass(), UMovieSceneAnimationMaskDecoration::StaticClass());
 	IMovieSceneModule::Get().RegisterCompatibleDecoration(UMovieSceneAnimationMixerLayer::StaticClass(), UMovieSceneLayerWeightDecoration::StaticClass());
@@ -65,6 +67,7 @@ void FMovieSceneAnimMixerModule::ShutdownModule()
 	{
 		IMovieSceneModule::Get().UnregisterCompatibleDecoration(UMovieSceneAnimationMixerTrack::StaticClass(), UMovieSceneRootMotionTargetDecoration::StaticClass());
 		IMovieSceneModule::Get().UnregisterCompatibleDecoration(UMovieSceneSkeletalAnimationSection::StaticClass(), UMovieSceneRootMotionSettingsDecoration::StaticClass());
+		IMovieSceneModule::Get().UnregisterCompatibleDecoration(UMovieSceneAnimBusSection::StaticClass(), UMovieSceneRootMotionSettingsDecoration::StaticClass());
 		IMovieSceneModule::Get().UnregisterCompatibleDecoration(UMovieSceneAnimationMixerLayer::StaticClass(), UMovieSceneAnimationMaskDecoration::StaticClass());
 		IMovieSceneModule::Get().UnregisterCompatibleDecoration(UMovieSceneAnimationMixerLayer::StaticClass(), UMovieSceneLayerWeightDecoration::StaticClass());
 	}
diff --git a/Engine/Plugins/Experimental/MovieSceneAnimMixer/Source/MovieSceneAnimMixer/Private/Systems/MovieSceneAnimBusSystem.cpp b/Engine/Plugins/Experimental/MovieSceneAnimMixer/Source/MovieSceneAnimMixer/Private/Systems/MovieSceneAnimBusSystem.cpp
new file mode 100644
--- /dev/null
+++ b/Engine/Plugins/Experimental/MovieSceneAnimMixer/Source/MovieSceneAnimMixer/Private/Systems/MovieSceneAnimBusSystem.cpp
@@ -0,0 +1,119 @@
+// Copyright Epic Games, Inc. All Rights Reserved.
+
+#include "Systems/MovieSceneAnimBusSystem.h"
+
+#include "AnimMixerComponentTypes.h"
+#include "EntitySystem/BuiltInComponentTypes.h"
+#include "EntitySystem/IMovieSceneTaskScheduler.h"
+#include "EntitySystem/MovieSceneEntitySystemLinker.h"
+#include "EntitySystem/MovieSceneEntitySystemTask.h"
+#include "EntitySystem/Interrogation/MovieSceneInterrogationLinker.h"
+#include "EvaluationVM/EvaluationVM.h"
+#include "EvaluationVM/KeyframeState.h"
+#include "EvaluationVM/EvaluationProgram.h"
+#include "Systems/MovieSceneAnimMixerSystem.h"
+
+#include UE_INLINE_GENERATED_CPP_BY_NAME(MovieSceneAnimBusSystem)
+
+FMovieSceneAnimBusReadTask FMovieSceneAnimBusReadTask::Make(FName InBusName, FObjectKey InBoundObjectKey, UMovieSceneAnimMixerSystem* InMixerSystem)
+{
+	FMovieSceneAnimBusReadTask Task;
+	Task.BusName = InBusName;
+	Task.BoundObjectKey = InBoundObjectKey;
+	Task.MixerSystem = InMixerSystem;
+	return Task;
+}
+
+void FMovieSceneAnimBusReadTask::Execute(UE::UAF::FEvaluationVM& VM) const
+{
+	using namespace UE::UAF;
+
+	if (!MixerSystem.IsValid())
+	{
+		FKeyframeState RefKeyframe = VM.MakeReferenceKeyframe(false);
+		VM.PushValue(KEYFRAME_STACK_NAME, MakeUnique<FKeyframeState>(MoveTemp(RefKeyframe)));
+		return;
+	}
+
+	TSharedPtr<FMovieSceneAnimBusData> BusData = MixerSystem->ReadBusData(BoundObjectKey, BusName);
+	if (BusData && BusData->IsValid())
+	{
+		// Execute the bus-target mixer's program directly in this VM.
+		// The program pushes its result onto the keyframe stack.
+		BusData->Program->Execute(VM);
+	}
+	else
+	{
+		// Bus has no writer this frame -- push identity pose
+		FKeyframeState RefKeyframe = VM.MakeReferenceKeyframe(false);
+		VM.PushValue(KEYFRAME_STACK_NAME, MakeUnique<FKeyframeState>(MoveTemp(RefKeyframe)));
+	}
+}
+
+namespace UE::MovieScene
+{
+
+struct FPopulateBusReadTasks
+{
+	UMovieSceneAnimMixerSystem* MixerSystem;
+
+	FPopulateBusReadTasks(UMovieSceneAnimMixerSystem* InMixerSystem)
+		: MixerSystem(InMixerSystem)
+	{
+	}
+
+	void ForEachEntity(
+		FObjectKey BoundObjectKey,
+		FName BusName,
+		TSharedPtr<FAnimNextEvaluationTask>& OutTask) const
+	{
+		if (BusName.IsNone() || !MixerSystem)
+		{
+			return;
+		}
+
+		OutTask = MakeShared<FMovieSceneAnimBusReadTask>(
+			FMovieSceneAnimBusReadTask::Make(BusName, BoundObjectKey, MixerSystem));
+	}
+};
+
+} // namespace UE::MovieScene
+
+UMovieSceneAnimBusSystem::UMovieSceneAnimBusSystem(const FObjectInitializer& ObjInit)
+	: Super(ObjInit)
+{
+	using namespace UE::MovieScene;
+
+	FAnimMixerComponentTypes* AnimMixerComponents = FAnimMixerComponentTypes::Get();
+
+	RelevantComponent = AnimMixerComponents->BusName;
+	Phase = ESystemPhase::Scheduling;
+
+	SystemCategories |= FSystemInterrogator::GetExcludedFromInterrogationCategory();
+
+	if (HasAnyFlags(RF_ClassDefaultObject))
+	{
+		// Run before the mixer system so bus read tasks are in place
+		DefineImplicitPrerequisite(GetClass(), UMovieSceneAnimMixerSystem::StaticClass());
+	}
+}
+
+void UMovieSceneAnimBusSystem::OnSchedulePersistentTasks(UE::MovieScene::IEntitySystemScheduler* TaskScheduler)
+{
+	using namespace UE::MovieScene;
+
+	FBuiltInComponentTypes* BuiltInComponents = FBuiltInComponentTypes::Get();
+	FAnimMixerComponentTypes* AnimMixerComponents = FAnimMixerComponentTypes::Get();
+
+	UMovieSceneAnimMixerSystem* MixerSystem = Linker->FindSystem<UMovieSceneAnimMixerSystem>();
+
+	FTaskParams Params(TEXT("Populate Bus Read Tasks"));
+	Params.ForceGameThread();
+	FEntityTaskBuilder()
+		.Read(BuiltInComponents->BoundObjectKey)
+		.Read(AnimMixerComponents->BusName)
+		.Write(AnimMixerComponents->Task)
+		.FilterNone({ BuiltInComponents->Tags.NeedsUnlink })
+		.SetParams(Params)
+		.Schedule_PerEntity<FPopulateBusReadTasks>(&Linker->EntityManager, TaskScheduler, MixerSystem);
+}
diff --git a/Engine/Plugins/Experimental/MovieSceneAnimMixer/Source/MovieSceneAnimMixer/Private/Systems/MovieSceneAnimMixerSystem.cpp b/Engine/Plugins/Experimental/MovieSceneAnimMixer/Source/MovieSceneAnimMixer/Private/Systems/MovieSceneAnimMixerSystem.cpp
--- a/Engine/Plugins/Experimental/MovieSceneAnimMixer/Source/MovieSceneAnimMixer/Private/Systems/MovieSceneAnimMixerSystem.cpp
+++ b/Engine/Plugins/Experimental/MovieSceneAnimMixer/Source/MovieSceneAnimMixer/Private/Systems/MovieSceneAnimMixerSystem.cpp
@@ -54,6 +54,8 @@
 #include "EvaluationVM/Tasks/TransitionEvaluationTask.h"
 #include "EntitySystem/MovieSceneInstanceRegistry.h"
 #include "MovieSceneAnimationMixerLayer.h"
+#include "MovieSceneAnimBusSection.h"
+#include "AnimMixerBusUtils.h"
 #include UE_INLINE_GENERATED_CPP_BY_NAME(MovieSceneAnimMixerSystem)
 
 namespace UE::MovieScene
@@ -76,6 +78,12 @@ struct FProgramBuildFilter
 
 	// When set, the conversion task writes the pre-conversion animation-space root here
 	TSharedPtr<FTransform> CaptureAnimSpaceRoot;
+
+	// When true, the first layer's inter-layer blend is skipped. The program
+	// produces a self-contained keyframe without consuming a base from the
+	// stack. Used for bus-target mixers whose programs are embedded inside
+	// a consuming mixer's entry evaluation.
+	bool bSkipFirstInterLayerBlend = false;
 };
 
 // Context needed for building evaluation programs.
@@ -309,7 +317,8 @@ static TSharedPtr<UE::UAF::FEvaluationProgram> BuildProgramFromEntries(
 	}
 
 	USceneComponent* Component = Context.Component;
-	USceneComponent* Root = Component ? Component->GetOwner()->GetRootComponent() : nullptr;
+	AActor* Owner = Component ? Component->GetOwner() : nullptr;
+	USceneComponent* Root = Owner ? Owner->GetRootComponent() : nullptr;
 	const bool bRootComponentHasKeyedTransform = (Root && Context.RootMotionSystem && Context.RootMotionSystem->IsTransformKeyed(Root));
 
 	// Second pass: handle transitions. A transition is only valid if BOTH its from and to
@@ -551,6 +560,7 @@ static TSharedPtr<UE::UAF::FEvaluationProgram> BuildProgramFromEntries(
 	};
 
 	// Main evaluation loop- two-level: intra-layer, then inter-layer blend.
+	bool bFirstLayer = true;
 	for (int32 Index = 0; Index < NumEntries; )
 	{
 		const TSharedPtr<FMovieSceneAnimMixerEntry>& PeekEntry = FinalEntries[Index];
@@ -673,36 +683,44 @@ static TSharedPtr<UE::UAF::FEvaluationProgram> BuildProgramFromEntries(
 
 		// ---- Level 2: Inter-layer blend ----
 		// Blend this layer's result with what's below on the stack.
+		// For bus-target programs, skip the first layer's blend since the
+		// program will be executed inline inside a consuming mixer that
+		// handles the blend with its own base keyframe.
+		const bool bSkipBlend = bFirstLayer && Filter.bSkipFirstInterLayerBlend;
+		bFirstLayer = false;
 
-		// Apply decoration-provided layer weight (e.g. from a Layer Weight decoration)
-		if (Layer.IsValid() && Context.DecorationLayerWeightMap)
+		if (!bSkipBlend)
 		{
-			if (const double* DecWeight = Context.DecorationLayerWeightMap->Find(Layer))
+			// Apply decoration-provided layer weight (e.g. from a Layer Weight decoration)
+			if (Layer.IsValid() && Context.DecorationLayerWeightMap)
 			{
-				LayerWeight *= *DecWeight;
+				if (const double* DecWeight = Context.DecorationLayerWeightMap->Find(Layer))
+				{
+					LayerWeight *= *DecWeight;
+				}
 			}
-		}
 
-		TSharedPtr<FAnimNextEvaluationTask> LayerBlendTask;
-		if (Layer.IsValid() && Context.LayerBlendTaskMap)
-		{
-			if (const FLayerBlendTaskBuilder* Builder = Context.LayerBlendTaskMap->Find(Layer))
+			TSharedPtr<FAnimNextEvaluationTask> LayerBlendTask;
+			if (Layer.IsValid() && Context.LayerBlendTaskMap)
 			{
-				if (Builder->IsValid())
+				if (const FLayerBlendTaskBuilder* Builder = Context.LayerBlendTaskMap->Find(Layer))
 				{
-					LayerBlendTask = Builder->Build(LayerWeight);
+					if (Builder->IsValid())
+					{
+						LayerBlendTask = Builder->Build(LayerWeight);
+					}
 				}
 			}
-		}
 
-		if (LayerBlendTask)
-		{
-			Program->AppendTaskPtr(LayerBlendTask);
-			Program->AppendTask(FAnimNextNormalizeKeyframeRotationsTask());
-		}
-		else
-		{
-			Program->AppendTask(FAnimNextBlendTwoKeyframesPreserveRootMotionTask::Make(LayerWeight));
+			if (LayerBlendTask)
+			{
+				Program->AppendTaskPtr(LayerBlendTask);
+				Program->AppendTask(FAnimNextNormalizeKeyframeRotationsTask());
+			}
+			else
+			{
+				Program->AppendTask(FAnimNextBlendTwoKeyframesPreserveRootMotionTask::Make(LayerWeight));
+			}
 		}
 	}
 
@@ -844,6 +862,155 @@ namespace UE::MovieScene
 					}
 				}
 			);
+
+			// --- Bus evaluation: dependency-ordered pre-evaluation of bus-target mixers ---
+			// Skip the full dependency graph build when no bus targets are in use.
+			if (HasAnyBusTarget())
+			{
+				EvaluateBusTargetMixers();
+			}
+			else
+			{
+				MixerSystem->ClearBusStorage();
+			}
+		}
+
+		bool HasAnyBusTarget() const
+		{
+			for (const auto& MixerPair : *Mixers)
+			{
+				if (MixerPair.Key.Target.GetPtr<FMovieSceneAnimBusTarget>())
+				{
+					return true;
+				}
+			}
+			return false;
+		}
+
+		void EvaluateBusTargetMixers()
+		{
+			MixerSystem->ClearBusStorage();
+
+			// Identify objects with bus targets and mark dirty caches
+			TSet<FObjectKey> ObjectsWithBusTargets;
+			for (const auto& MixerPair : *Mixers)
+			{
+				if (MixerPair.Key.Target.GetPtr<FMovieSceneAnimBusTarget>())
+				{
+					FObjectKey ObjKey = MixerPair.Key.BoundObjectKey;
+					ObjectsWithBusTargets.Add(ObjKey);
+
+					if (TSharedPtr<FMovieSceneAnimMixer> Mixer = MixerPair.Value)
+					{
+						if (Mixer->bNeedsResort)
+						{
+							MixerSystem->GetOrCreateBusTopology(ObjKey).bDirty = true;
+						}
+					}
+				}
+			}
+
+			MixerSystem->RemoveStaleBusTopologies(ObjectsWithBusTargets);
+
+			for (FObjectKey ObjKey : ObjectsWithBusTargets)
+			{
+				auto& Cache = MixerSystem->GetOrCreateBusTopology(ObjKey);
+				if (Cache.bDirty)
+				{
+					// Gather owner tracks for this object only when we need to rebuild
+					TArray<UMovieSceneAnimationMixerTrack*> ObjectTracks;
+					for (const auto& MixerPair : *Mixers)
+					{
+						if (MixerPair.Key.BoundObjectKey == ObjKey)
+						{
+							if (TSharedPtr<FMovieSceneAnimMixer> Mixer = MixerPair.Value)
+							{
+								if (Mixer->OwnerTrack.IsValid())
+								{
+									ObjectTracks.AddUnique(Mixer->OwnerTrack.Get());
+								}
+							}
+						}
+					}
+					Cache.EvalOrder = FAnimMixerBusUtils::ComputeBusEvaluationOrder(ObjectTracks);
+					Cache.bDirty = false;
+				}
+
+				// Build programs and store bus data in cached order
+				for (const FName& BusName : Cache.EvalOrder)
+				{
+					FMovieSceneAnimMixerKey MixerKey;
+					MixerKey.BoundObjectKey = ObjKey;
+					MixerKey.Target = TInstancedStruct<FMovieSceneAnimBusTarget>::Make(BusName);
+
+					TSharedPtr<FMovieSceneAnimMixer> Mixer = Mixers->FindRef(MixerKey);
+					if (!Mixer)
+					{
+						continue;
+					}
+
+					USceneComponent* Component = Cast<USceneComponent>(ObjKey.ResolveObjectPtr());
+					TSharedPtr<FMovieSceneMixerRootMotionComponentData> RootMotion = RootMotions->FindRef(ObjKey).Pin();
+
+					FProgramBuildContext BuildContext;
+					BuildContext.RootMotion = RootMotion;
+					BuildContext.RootMotionSystem = RootMotionSystem;
+					BuildContext.TransformOriginSystem = TransformOriginSystem;
+					BuildContext.Linker = Linker;
+					BuildContext.Component = Component;
+					BuildContext.LayerBlendTaskMap = LayerBlendBuilders.Find(ObjKey);
+					BuildContext.DecorationLayerWeightMap = DecorationLayerWeights.Find(ObjKey);
+
+					FProgramBuildFilter BuildFilter;
+					BuildFilter.bSkipFirstInterLayerBlend = true;
+					Mixer->EvaluationProgram = BuildProgramFromEntries(*Mixer, BuildContext, BuildFilter);
+
+					// Check root motion per-frame since entries can change
+					bool bHasRootMotion = false;
+					if (Mixer->EvaluationProgram && RootMotion)
+					{
+						for (const TWeakPtr<FMovieSceneAnimMixerEntry>& WeakEntry : Mixer->WeakEntries)
+						{
+							if (TSharedPtr<FMovieSceneAnimMixerEntry> Entry = WeakEntry.Pin())
+							{
+								if (Entry->RootMotionSettings.IsSet())
+								{
+									bHasRootMotion = true;
+									break;
+								}
+								if (UMovieSceneAnimBusSection* BusSection = Cast<UMovieSceneAnimBusSection>(Entry->EntityOwner.ResolveObjectPtr()))
+								{
+									TSharedPtr<FMovieSceneAnimBusData> BusData = MixerSystem->ReadBusData(ObjKey, BusSection->BusName);
+									if (BusData && BusData->bHasRootMotion)
+									{
+										bHasRootMotion = true;
+										break;
+									}
+								}
+							}
+						}
+
+						if (bHasRootMotion)
+						{
+							AActor* Owner = Component ? Component->GetOwner() : nullptr;
+							USceneComponent* Root = Owner ? Owner->GetRootComponent() : nullptr;
+							const bool bComponentHasKeyedTransform = (Component && RootMotionSystem && RootMotionSystem->IsTransformKeyed(Component));
+							const bool bRootComponentHasKeyedTransform = (Root && RootMotionSystem && RootMotionSystem->IsTransformKeyed(Root));
+							FAnimNextStoreRootTransformTask StoreTask = FAnimNextStoreRootTransformTask::Make(RootMotion, bComponentHasKeyedTransform, bRootComponentHasKeyedTransform);
+							StoreTask.bForceRootBoneDestination = true;
+							Mixer->EvaluationProgram->AppendTask(MoveTemp(StoreTask));
+						}
+					}
+
+					if (Mixer->EvaluationProgram)
+					{
+						TSharedPtr<FMovieSceneAnimBusData> BusData = MakeShared<FMovieSceneAnimBusData>();
+						BusData->Program = Mixer->EvaluationProgram;
+						BusData->bHasRootMotion = bHasRootMotion;
+						MixerSystem->StoreBusData(ObjKey, BusName, BusData);
+					}
+				}
+			}
 		}
 
 		// Get an InstanceHandle from the mixer's entries (any valid entry will do).
@@ -869,6 +1036,13 @@ namespace UE::MovieScene
 			TSharedPtr<FAnimNextEvaluationTask>& MixerTask
 			) const
 		{
+			// Bus-target mixers are handled in EvaluateBusTargetMixers. Their
+			// programs are stored in BusStorage, not on the MixerTask entity.
+			if (Target.GetPtr<FMovieSceneAnimBusTarget>())
+			{
+				return;
+			}
+
 			TSharedPtr<FMovieSceneAnimMixer> Mixer = Mixers->FindRef(FMovieSceneAnimMixerKey({ MeshComponent.GetObject(), Target }));
 			if (!Mixer)
 			{
@@ -930,6 +1104,16 @@ namespace UE::MovieScene
 							bAnyEntryHasRootMotion = true;
 							break;
 						}
+						// Check if this entry is a bus section whose source has root motion
+						if (UMovieSceneAnimBusSection* BusSection = Cast<UMovieSceneAnimBusSection>(Entry->EntityOwner.ResolveObjectPtr()))
+						{
+							TSharedPtr<FMovieSceneAnimBusData> BusData = MixerSystem->ReadBusData(MeshComponent.GetObject(), BusSection->BusName);
+							if (BusData && BusData->bHasRootMotion)
+							{
+								bAnyEntryHasRootMotion = true;
+								break;
+							}
+						}
 					}
 				}
 
@@ -963,7 +1147,8 @@ namespace UE::MovieScene
 					// motion attribute directly from the result keyframe.
 					if (!Target.GetPtr<FMovieSceneMixerBakeTarget>())
 					{
-						USceneComponent* Root = Component ? Component->GetOwner()->GetRootComponent() : nullptr;
+						AActor* Owner = Component ? Component->GetOwner() : nullptr;
+						USceneComponent* Root = Owner ? Owner->GetRootComponent() : nullptr;
 						const bool bComponentHasKeyedTransform = (Component && RootMotionSystem->IsTransformKeyed(Component));
 						const bool bRootComponentHasKeyedTransform = (Root && RootMotionSystem->IsTransformKeyed(Root));
 						Mixer->EvaluationProgram->AppendTask(FAnimNextStoreRootTransformTask::Make(RootMotion, bComponentHasKeyedTransform, bRootComponentHasKeyedTransform));
@@ -1271,6 +1456,22 @@ void UMovieSceneAnimMixerSystem::PreInitializeAllRootMotion()
 	}
 }
 
+UMovieSceneAnimMixerSystem::FBusTopologyCache& UMovieSceneAnimMixerSystem::GetOrCreateBusTopology(FObjectKey ObjectKey)
+{
+	return BusTopologyByObject.FindOrAdd(ObjectKey);
+}
+
+void UMovieSceneAnimMixerSystem::RemoveStaleBusTopologies(const TSet<FObjectKey>& ActiveObjects)
+{
+	for (auto It = BusTopologyByObject.CreateIterator(); It; ++It)
+	{
+		if (!ActiveObjects.Contains(It.Key()))
+		{
+			It.RemoveCurrent();
+		}
+	}
+}
+
 void UMovieSceneAnimMixerSystem::InitializeAllRootMotion()
 {
 	for (auto It = RootMotionByObject.CreateIterator(); It; ++It)
@@ -1283,6 +1484,28 @@ void UMovieSceneAnimMixerSystem::InitializeAllRootMotion()
 	}
 }
 
+void UMovieSceneAnimMixerSystem::StoreBusData(FObjectKey BoundObject, FName BusName, TSharedPtr<FMovieSceneAnimBusData> Data)
+{
+	FMovieSceneAnimBusStorageKey Key;
+	Key.BoundObjectKey = BoundObject;
+	Key.BusName = BusName;
+	BusStorage.Add(Key, MoveTemp(Data));
+}
+
+TSharedPtr<FMovieSceneAnimBusData> UMovieSceneAnimMixerSystem::ReadBusData(FObjectKey BoundObject, FName BusName) const
+{
+	FMovieSceneAnimBusStorageKey Key;
+	Key.BoundObjectKey = BoundObject;
+	Key.BusName = BusName;
+	const TSharedPtr<FMovieSceneAnimBusData>* Found = BusStorage.Find(Key);
+	return Found ? *Found : nullptr;
+}
+
+void UMovieSceneAnimMixerSystem::ClearBusStorage()
+{
+	BusStorage.Reset();
+}
+
 bool UMovieSceneAnimMixerSystem::IsRelevantImpl(UMovieSceneEntitySystemLinker* InLinker) const
 {
 	return Mixers.Num() != 0;
@@ -1352,8 +1575,14 @@ void UMovieSceneAnimMixerSystem::HandlePreFlush(UMovieSceneEntitySystemLinker* I
 			{
 				if (Entry->InstanceHandle.IsValid())
 				{
+					// Resolve to the root instance handle. The bake evaluation
+					// flushes the ECS which only supports root instance updates.
+					FInstanceHandle Handle = Entry->InstanceHandle;
+					const FSequenceInstance& Instance = InLinker->GetInstanceRegistry()->GetInstance(Handle);
+					FRootInstanceHandle RootHandle = Instance.GetRootInstanceHandle();
+
 					DirtyTracks.Add(Track);
-					DirtyHandles.Add(Entry->InstanceHandle);
+					DirtyHandles.Add(RootHandle);
 					break;
 				}
 			}
diff --git a/Engine/Plugins/Experimental/MovieSceneAnimMixer/Source/MovieSceneAnimMixer/Private/Systems/MovieSceneRootMotionSystem.cpp b/Engine/Plugins/Experimental/MovieSceneAnimMixer/Source/MovieSceneAnimMixer/Private/Systems/MovieSceneRootMotionSystem.cpp
--- a/Engine/Plugins/Experimental/MovieSceneAnimMixer/Source/MovieSceneAnimMixer/Private/Systems/MovieSceneRootMotionSystem.cpp
+++ b/Engine/Plugins/Experimental/MovieSceneAnimMixer/Source/MovieSceneAnimMixer/Private/Systems/MovieSceneRootMotionSystem.cpp
@@ -439,7 +439,9 @@ void FAnimNextStoreRootTransformTask::Execute(UE::UAF::FEvaluationVM& VM) const
 		return;
 	}
 
-	const EMovieSceneRootMotionDestination FinalDestination = RootMotionData->RootDestination;
+	const EMovieSceneRootMotionDestination FinalDestination = bForceRootBoneDestination
+		? EMovieSceneRootMotionDestination::RootBone
+		: RootMotionData->RootDestination;
 
 	FTransform RootTransform = RootMotionAttribute->Value;
 
@@ -469,6 +471,15 @@ void FAnimNextStoreRootTransformTask::Execute(UE::UAF::FEvaluationVM& VM) const
 		if (RootBoneIndex != INDEX_NONE)
 		{
 			(*Keyframe)->Pose.LocalTransformsView[RootBoneIndex] = RootTransform;
+
+			if (bForceRootBoneDestination)
+			{
+				// Remove the attribute so the outer mixer's extract can cleanly
+				// pull root motion from the root bone without double-applying.
+				(*Keyframe)->Attributes.Remove<FTransformAnimationAttribute>(AnimMixerComponents->RootTransformAttributeId);
+				(*Keyframe)->Attributes.Remove<FFloatAnimationAttribute>(AnimMixerComponents->RootTransformWeightAttributeId);
+				(*Keyframe)->Attributes.Remove<FIntegerAnimationAttribute>(AnimMixerComponents->RootTransformIsAuthoritativeAttributeId);
+			}
 		}
 		return;
 
diff --git a/Engine/Plugins/Experimental/MovieSceneAnimMixer/Source/MovieSceneAnimMixer/Public/AnimMixerBusUtils.h b/Engine/Plugins/Experimental/MovieSceneAnimMixer/Source/MovieSceneAnimMixer/Public/AnimMixerBusUtils.h
new file mode 100644
--- /dev/null
+++ b/Engine/Plugins/Experimental/MovieSceneAnimMixer/Source/MovieSceneAnimMixer/Public/AnimMixerBusUtils.h
@@ -0,0 +1,65 @@
+// Copyright Epic Games, Inc. All Rights Reserved.
+
+#pragma once
+
+#include "CoreMinimal.h"
+
+class UMovieSceneSequence;
+class UMovieSceneAnimationMixerTrack;
+
+// Dependency graph for bus-target mixers. Callers populate Nodes, then
+// call TopologicalSort or IsReachable.
+struct MOVIESCENEANIMMIXER_API FBusDependencyGraph
+{
+	struct FNode
+	{
+		FName WritesToBus;
+		TArray<FName> ReadsFromBuses;
+
+		// Set by TopologicalSort
+		int32 InDegree = 0;
+		TArray<int32> Dependents;
+	};
+
+	TArray<FNode> Nodes;
+
+	// Kahn's algorithm. Returns sorted node indices; nodes in cycles keep
+	// InDegree > 0 and are excluded. Self-references are skipped.
+	TArray<int32> TopologicalSort();
+
+	// True if TargetBus is reachable from StartBus through dependency edges.
+	bool IsReachable(FName StartBus, FName TargetBus) const;
+};
+
+struct MOVIESCENEANIMMIXER_API FAnimMixerBusValidationResult
+{
+	TArray<FString> Errors;
+	TArray<FString> Warnings;
+
+	bool HasErrors() const { return !Errors.IsEmpty(); }
+	bool HasWarnings() const { return !Warnings.IsEmpty(); }
+	bool IsClean() const { return Errors.IsEmpty() && Warnings.IsEmpty(); }
+};
+
+// Bus topology utilities. Functions that take mixer tracks expect all tracks
+// for a single bound object.
+struct MOVIESCENEANIMMIXER_API FAnimMixerBusUtils
+{
+	// Gather all bus names across the sequence hierarchy
+	static TArray<FName> GatherBusNamesFromSequence(UMovieSceneSequence* RootSequence);
+
+	// Build a dependency graph from mixer tracks and return topologically
+	// sorted bus names. Cyclic buses are excluded (with an ensure).
+	static TArray<FName> ComputeBusEvaluationOrder(TArrayView<UMovieSceneAnimationMixerTrack* const> MixerTracks);
+
+	// Validate bus topology: cycles, self-references, duplicate writers, orphans.
+	static FAnimMixerBusValidationResult ValidateBusTopology(TArrayView<UMovieSceneAnimationMixerTrack* const> MixerTracks);
+
+	// Would adding a bus section reading BusName to Track create a cycle?
+	static bool WouldBusSectionCreateCycle(FName BusName, UMovieSceneAnimationMixerTrack* Track,
+		TArrayView<UMovieSceneAnimationMixerTrack* const> MixerTracks);
+
+	// Would setting Track's target to BusName create a cycle?
+	static bool WouldBusTargetCreateCycle(FName BusName, UMovieSceneAnimationMixerTrack* Track,
+		TArrayView<UMovieSceneAnimationMixerTrack* const> MixerTracks);
+};
diff --git a/Engine/Plugins/Experimental/MovieSceneAnimMixer/Source/MovieSceneAnimMixer/Public/AnimMixerComponentTypes.h b/Engine/Plugins/Experimental/MovieSceneAnimMixer/Source/MovieSceneAnimMixer/Public/AnimMixerComponentTypes.h
--- a/Engine/Plugins/Experimental/MovieSceneAnimMixer/Source/MovieSceneAnimMixer/Public/AnimMixerComponentTypes.h
+++ b/Engine/Plugins/Experimental/MovieSceneAnimMixer/Source/MovieSceneAnimMixer/Public/AnimMixerComponentTypes.h
@@ -70,6 +70,9 @@ namespace UE::MovieScene
 		// The builder is called with the computed layer weight and returns the blend task.
 		TComponentTypeID<FLayerBlendTaskBuilder> LayerBlendTask;
 
+		// Bus name for bus sections. Identifies which named bus this section reads from.
+		TComponentTypeID<FName> BusName;
+
 		struct
 		{
 			FComponentTypeID Transition;
@@ -77,7 +80,6 @@ namespace UE::MovieScene
 			FComponentTypeID LayerWeight;
 		} Tags;
 
-
 	private:
 		FAnimMixerComponentTypes();
 	};
diff --git a/Engine/Plugins/Experimental/MovieSceneAnimMixer/Source/MovieSceneAnimMixer/Public/MovieSceneAnimBusSection.h b/Engine/Plugins/Experimental/MovieSceneAnimMixer/Source/MovieSceneAnimMixer/Public/MovieSceneAnimBusSection.h
new file mode 100644
--- /dev/null
+++ b/Engine/Plugins/Experimental/MovieSceneAnimMixer/Source/MovieSceneAnimMixer/Public/MovieSceneAnimBusSection.h
@@ -0,0 +1,48 @@
+// Copyright Epic Games, Inc. All Rights Reserved.
+
+#pragma once
+
+#include "Channels/MovieSceneFloatChannel.h"
+#include "MovieSceneAnimationMixerItemInterface.h"
+#include "MovieSceneSection.h"
+#include "EntitySystem/IMovieSceneEntityProvider.h"
+
+#include "MovieSceneAnimBusSection.generated.h"
+
+// Section type for reading a stored pose from a named bus into a mixer layer.
+// Place this on a mixer layer to pull in the pose that another mixer wrote to
+// the specified bus via a bus target. Defaults to infinite range but can be
+// trimmed to only pull from the bus during part of the sequence.
+UCLASS(MinimalAPI, DisplayName="Bus")
+class UMovieSceneAnimBusSection
+	: public UMovieSceneSection
+	, public IMovieSceneEntityProvider
+	, public IMovieSceneAnimationMixerItemInterface
+{
+public:
+
+	UMovieSceneAnimBusSection();
+
+	GENERATED_BODY()
+
+	// Name of the bus to read the pose from
+	UPROPERTY(EditAnywhere, Category="Animation Bus")
+	FName BusName;
+
+	// Weight channel controlling how much the bus pose contributes to the mix
+	UPROPERTY()
+	FMovieSceneFloatChannel Weight;
+
+
+	virtual FColor GetMixerItemTint() const override
+	{
+		return BusSectionTintColor;
+	}
+
+private:
+
+	virtual EMovieSceneChannelProxyType CacheChannelProxy() override;
+	virtual void ImportEntityImpl(UMovieSceneEntitySystemLinker* EntityLinker, const FEntityImportParams& Params, FImportedEntity* OutImportedEntity) override;
+
+	static constexpr FColor BusSectionTintColor = FColor(140, 70, 180);
+};
diff --git a/Engine/Plugins/Experimental/MovieSceneAnimMixer/Source/MovieSceneAnimMixer/Public/Systems/MovieSceneAnimBusSystem.h b/Engine/Plugins/Experimental/MovieSceneAnimMixer/Source/MovieSceneAnimMixer/Public/Systems/MovieSceneAnimBusSystem.h
new file mode 100644
--- /dev/null
+++ b/Engine/Plugins/Experimental/MovieSceneAnimMixer/Source/MovieSceneAnimMixer/Public/Systems/MovieSceneAnimBusSystem.h
@@ -0,0 +1,50 @@
+// Copyright Epic Games, Inc. All Rights Reserved.
+
+#pragma once
+
+#include "CoreTypes.h"
+#include "UObject/ObjectKey.h"
+#include "EntitySystem/MovieSceneEntitySystem.h"
+#include "EvaluationVM/EvaluationTask.h"
+#include "MovieSceneAnimBusSystem.generated.h"
+
+class UMovieSceneAnimMixerSystem;
+
+// Evaluation task that reads a stored pose from a named bus and pushes it
+// onto the VM's keyframe stack. Used by UMovieSceneAnimBusSection to inject
+// poses written by bus-target mixers into a consuming mixer's program.
+USTRUCT()
+struct MOVIESCENEANIMMIXER_API FMovieSceneAnimBusReadTask : public FAnimNextEvaluationTask
+{
+	GENERATED_BODY()
+
+	DECLARE_ANIM_EVALUATION_TASK(FMovieSceneAnimBusReadTask)
+
+	static FMovieSceneAnimBusReadTask Make(FName InBusName, FObjectKey InBoundObjectKey, UMovieSceneAnimMixerSystem* InMixerSystem);
+
+	virtual void Execute(UE::UAF::FEvaluationVM& VM) const override;
+
+	UPROPERTY()
+	FName BusName;
+
+	FObjectKey BoundObjectKey;
+
+	TWeakObjectPtr<UMovieSceneAnimMixerSystem> MixerSystem;
+};
+
+// System that creates/updates FMovieSceneAnimBusReadTask instances on entities
+// with BusName components. Runs before UMovieSceneAnimMixerSystem so bus read
+// tasks are in place before the mixer builds its evaluation program.
+UCLASS(MinimalAPI)
+class UMovieSceneAnimBusSystem : public UMovieSceneEntitySystem
+{
+public:
+
+	GENERATED_BODY()
+
+	UMovieSceneAnimBusSystem(const FObjectInitializer& ObjInit);
+
+private:
+
+	virtual void OnSchedulePersistentTasks(UE::MovieScene::IEntitySystemScheduler* TaskScheduler) override;
+};
diff --git a/Engine/Plugins/Experimental/MovieSceneAnimMixer/Source/MovieSceneAnimMixer/Public/Systems/MovieSceneAnimMixerSystem.h b/Engine/Plugins/Experimental/MovieSceneAnimMixer/Source/MovieSceneAnimMixer/Public/Systems/MovieSceneAnimMixerSystem.h
--- a/Engine/Plugins/Experimental/MovieSceneAnimMixer/Source/MovieSceneAnimMixer/Public/Systems/MovieSceneAnimMixerSystem.h
+++ b/Engine/Plugins/Experimental/MovieSceneAnimMixer/Source/MovieSceneAnimMixer/Public/Systems/MovieSceneAnimMixerSystem.h
@@ -11,6 +11,7 @@
 #include "AnimSequencerInstanceProxy.h"
 #include "Systems/MovieSceneRootMotionSystem.h"
 #include "EvaluationVM/Tasks/BlendKeyframes.h"
+#include "EvaluationVM/EvaluationProgram.h"
 #include "MovieSceneAnimMixerSystem.generated.h"
 
 class UMovieSceneAnimationMixerTrack;
@@ -23,6 +24,43 @@ namespace UE::UAF
 	struct FEvaluationProgram;
 }
 
+// Target type that writes the mixer's evaluated pose to a named bus.
+// Other mixers can then read from the bus via UMovieSceneAnimBusSection.
+USTRUCT(meta=(DisplayName="Bus Target"))
+struct MOVIESCENEANIMMIXER_API FMovieSceneAnimBusTarget : public FMovieSceneMixedAnimationTarget
+{
+	GENERATED_BODY()
+
+	UPROPERTY(EditAnywhere, Category="Animation")
+	FName BusName;
+
+	inline friend uint32 GetTypeHash(const FMovieSceneAnimBusTarget& Target)
+	{
+		return HashCombine(GetTypeHash(FMovieSceneAnimBusTarget::StaticStruct()), GetTypeHash(Target.BusName));
+	}
+
+	FMovieSceneAnimBusTarget(FName InBusName)
+		: BusName(InBusName)
+	{
+	}
+
+	FMovieSceneAnimBusTarget()
+		: BusName(NAME_None)
+	{
+	}
+
+#if WITH_EDITOR
+	virtual FText GetShortDisplayName() const override
+	{
+		if (!BusName.IsNone())
+		{
+			return FText::Format(NSLOCTEXT("MovieSceneAnimBusTarget", "BusDisplayFmt", "Bus: {0}"), FText::FromName(BusName));
+		}
+		return NSLOCTEXT("MovieSceneAnimBusTarget", "BusTarget", "Bus Target");
+	}
+#endif
+};
+
 struct FMovieSceneAnimMixer;
 
 struct FMovieSceneAnimMixerEntry
@@ -167,6 +205,35 @@ struct FAnimNextBlendTwoKeyframesPreserveRootMotionTask : public FAnimNextBlendT
 
 
 
+// Key for per-object, per-bus-name storage.
+struct FMovieSceneAnimBusStorageKey
+{
+	FObjectKey BoundObjectKey;
+	FName BusName;
+
+	friend bool operator==(const FMovieSceneAnimBusStorageKey& A, const FMovieSceneAnimBusStorageKey& B)
+	{
+		return A.BoundObjectKey == B.BoundObjectKey && A.BusName == B.BusName;
+	}
+
+	friend uint32 GetTypeHash(const FMovieSceneAnimBusStorageKey& Key)
+	{
+		return HashCombineFast(GetTypeHash(Key.BoundObjectKey), GetTypeHash(Key.BusName));
+	}
+};
+
+// Stored evaluation program for a bus-target mixer. Executed inline in the
+// consuming mixer's VM so the pose is evaluated in the correct context.
+struct MOVIESCENEANIMMIXER_API FMovieSceneAnimBusData
+{
+	TSharedPtr<UE::UAF::FEvaluationProgram> Program;
+
+	// True if the bus-target mixer has entries with root motion settings.
+	bool bHasRootMotion = false;
+
+	bool IsValid() const { return Program.IsValid(); }
+};
+
 // Takes in evaluation tasks with optional pose weight, masks, priority and a given animation target.
 // Constructs a hierarchical 'mixer' per bound object per target. 
 // Similar to blender systems, in a 'many to one' operation, each mixer will create an entity with a single evaluation task
@@ -190,11 +257,38 @@ public:
 	void PreInitializeAllRootMotion();
 	void InitializeAllRootMotion();
 
+	// Bus storage API
+
+	// Store a pose result into the named bus for the given bound object
+	MOVIESCENEANIMMIXER_API void StoreBusData(FObjectKey BoundObject, FName BusName, TSharedPtr<FMovieSceneAnimBusData> Data);
+
+	// Read a stored bus pose for the given bound object and bus name (returns null if not written this frame)
+	MOVIESCENEANIMMIXER_API TSharedPtr<FMovieSceneAnimBusData> ReadBusData(FObjectKey BoundObject, FName BusName) const;
+
+	// Clear all bus storage (called at the start of each evaluation frame)
+	MOVIESCENEANIMMIXER_API void ClearBusStorage();
+
+	// Cached bus evaluation order per bound object
+	struct FBusTopologyCache
+	{
+		TArray<FName> EvalOrder;
+		bool bDirty = true;
+	};
+
+	FBusTopologyCache& GetOrCreateBusTopology(FObjectKey ObjectKey);
+	void RemoveStaleBusTopologies(const TSet<FObjectKey>& ActiveObjects);
+
 private:
 
 	// Map of animation mixers
 	TMap<FMovieSceneAnimMixerKey, TSharedPtr<FMovieSceneAnimMixer>> Mixers;
 
+	// Bus storage: per-object, per-bus-name pose data written by bus-target mixers
+	// and read by bus sections. Cleared at the start of each evaluation frame.
+	TMap<FMovieSceneAnimBusStorageKey, TSharedPtr<FMovieSceneAnimBusData>> BusStorage;
+
+	TMap<FObjectKey, FBusTopologyCache> BusTopologyByObject;
+
 	// Map of Root Motion data for each object
 	TMap<FObjectKey, TWeakPtr<FMovieSceneMixerRootMotionComponentData>> RootMotionByObject;
 
diff --git a/Engine/Plugins/Experimental/MovieSceneAnimMixer/Source/MovieSceneAnimMixer/Public/Systems/MovieSceneRootMotionSystem.h b/Engine/Plugins/Experimental/MovieSceneAnimMixer/Source/MovieSceneAnimMixer/Public/Systems/MovieSceneRootMotionSystem.h
--- a/Engine/Plugins/Experimental/MovieSceneAnimMixer/Source/MovieSceneAnimMixer/Public/Systems/MovieSceneRootMotionSystem.h
+++ b/Engine/Plugins/Experimental/MovieSceneAnimMixer/Source/MovieSceneAnimMixer/Public/Systems/MovieSceneRootMotionSystem.h
@@ -232,6 +232,11 @@ struct FAnimNextStoreRootTransformTask : public FAnimNextEvaluationTask
 	TWeakPtr<FMovieSceneMixerRootMotionComponentData> WeakRootMotionData;
 	bool bComponentHasKeyedTransform = false;
 	bool bRootComponentHasKeyedTransform = false;
+
+	// When true, forces RootBone destination regardless of the Root Motion Target
+	// settings. Used by bus-target mixers to put world-space root motion back onto
+	// the root bone so that outer mixers can re-extract and apply it.
+	bool bForceRootBoneDestination = false;
 };
 
 
diff --git a/Engine/Plugins/Experimental/MovieSceneAnimMixer/Source/MovieSceneAnimMixerEditor/Private/AnimBusSectionInterface.cpp b/Engine/Plugins/Experimental/MovieSceneAnimMixer/Source/MovieSceneAnimMixerEditor/Private/AnimBusSectionInterface.cpp
new file mode 100644
--- /dev/null
+++ b/Engine/Plugins/Experimental/MovieSceneAnimMixer/Source/MovieSceneAnimMixerEditor/Private/AnimBusSectionInterface.cpp
@@ -0,0 +1,176 @@
+// Copyright Epic Games, Inc. All Rights Reserved.
+
+#include "AnimBusSectionInterface.h"
+
+#include "AnimMixerBusUtils.h"
+#include "AnimMixerEditorBusUtils.h"
+#include "ISequencer.h"
+#include "MovieSceneAnimBusSection.h"
+#include "MovieSceneAnimationMixerTrack.h"
+#include "MovieSceneSequence.h"
+#include "Framework/Application/SlateApplication.h"
+#include "Framework/MultiBox/MultiBoxBuilder.h"
+#include "Framework/Notifications/NotificationManager.h"
+#include "Widgets/Notifications/SNotificationList.h"
+#include "ScopedTransaction.h"
+#include "Widgets/Input/SEditableTextBox.h"
+
+#define LOCTEXT_NAMESPACE "AnimBusSectionInterface"
+
+namespace UE::Sequencer
+{
+
+FAnimBusSectionInterface::FAnimBusSectionInterface(UMovieSceneSection& InSection, UMovieSceneTrack& InTrack, FGuid InObjectBinding, TWeakPtr<ISequencer> InSequencer)
+	: FSequencerSection(InSection)
+	, WeakTrack(&InTrack)
+	, WeakSequencer(InSequencer)
+	, ObjectBindingId(InObjectBinding)
+{
+}
+
+FText FAnimBusSectionInterface::GetSectionTitle() const
+{
+	UMovieSceneAnimBusSection* BusSection = GetBusSection();
+	if (BusSection && !BusSection->BusName.IsNone())
+	{
+		return FText::Format(LOCTEXT("BusSectionTitleFmt", "Bus: {0}"), FText::FromName(BusSection->BusName));
+	}
+	return LOCTEXT("BusSectionTitle", "Bus");
+}
+
+FText FAnimBusSectionInterface::GetSectionToolTip() const
+{
+	UMovieSceneAnimBusSection* BusSection = GetBusSection();
+	if (BusSection && !BusSection->BusName.IsNone())
+	{
+		return FText::Format(LOCTEXT("BusSectionTooltipFmt", "Reads pose from bus '{0}'"), FText::FromName(BusSection->BusName));
+	}
+	return LOCTEXT("BusSectionTooltip", "Bus section (no bus assigned)");
+}
+
+UMovieSceneAnimBusSection* FAnimBusSectionInterface::GetBusSection() const
+{
+	return Cast<UMovieSceneAnimBusSection>(WeakSection.Get());
+}
+
+UMovieSceneAnimationMixerTrack* FAnimBusSectionInterface::GetMixerTrack() const
+{
+	return Cast<UMovieSceneAnimationMixerTrack>(WeakTrack.Get());
+}
+
+void FAnimBusSectionInterface::BuildSectionContextMenu(FMenuBuilder& MenuBuilder, const FGuid& ObjectBinding)
+{
+	MenuBuilder.AddSubMenu(
+		LOCTEXT("ChangeBus", "Change Bus..."),
+		LOCTEXT("ChangeBusTooltip", "Change which bus this section reads from."),
+		FNewMenuDelegate::CreateSP(this, &FAnimBusSectionInterface::PopulateChangeBusMenu)
+	);
+}
+
+void FAnimBusSectionInterface::PopulateChangeBusMenu(FMenuBuilder& MenuBuilder)
+{
+	TWeakPtr<FAnimBusSectionInterface> WeakThis = SharedThis(this);
+	TSharedPtr<ISequencer> Sequencer = WeakSequencer.Pin();
+	UMovieSceneAnimationMixerTrack* MixerTrack = GetMixerTrack();
+	UMovieSceneSequence* RootSequence = Sequencer ? Sequencer->GetRootMovieSceneSequence() : nullptr;
+
+	// List known buses
+	TArray<FName> KnownBuses;
+	if (RootSequence)
+	{
+		KnownBuses = FAnimMixerBusUtils::GatherBusNamesFromSequence(RootSequence);
+	}
+
+	TArray<UMovieSceneAnimationMixerTrack*> TracksForObject = Sequencer
+		? AnimMixerEditorBusUtils::GatherMixerTracksForSameObject(MixerTrack, *Sequencer)
+		: TArray<UMovieSceneAnimationMixerTrack*>();
+
+	UMovieSceneAnimBusSection* BusSection = GetBusSection();
+	const FName CurrentBusName = BusSection ? BusSection->BusName : NAME_None;
+
+	for (const FName& BusName : KnownBuses)
+	{
+		if (BusName == CurrentBusName)
+		{
+			continue;
+		}
+
+		const bool bWouldCycle = MixerTrack &&
+			FAnimMixerBusUtils::WouldBusSectionCreateCycle(BusName, MixerTrack, TracksForObject);
+
+		MenuBuilder.AddMenuEntry(
+			FText::FromName(BusName),
+			bWouldCycle
+				? FText::Format(LOCTEXT("ChangeBusCycleTooltip", "Cannot read from bus '{0}': would create a dependency cycle."), FText::FromName(BusName))
+				: FText::Format(LOCTEXT("ChangeBusEntryTooltip", "Read from bus '{0}'."), FText::FromName(BusName)),
+			FSlateIcon(),
+			FUIAction(
+				FExecuteAction::CreateLambda([WeakThis, BusName]()
+				{
+					if (TSharedPtr<FAnimBusSectionInterface> This = WeakThis.Pin())
+					{
+						This->ChangeBusName(BusName);
+					}
+					FSlateApplication::Get().DismissAllMenus();
+				}),
+				FCanExecuteAction::CreateLambda([bWouldCycle]() { return !bWouldCycle; }))
+		);
+	}
+
+	// "New Bus..." text entry
+	TSharedRef<SEditableTextBox> TextBox = SNew(SEditableTextBox)
+		.HintText(LOCTEXT("BusNameHint", "Enter bus name..."))
+		.OnTextCommitted_Lambda([WeakThis, TracksForObject](const FText& InText, ETextCommit::Type CommitType)
+		{
+			TSharedPtr<FAnimBusSectionInterface> This = WeakThis.Pin();
+			if (!This)
+			{
+				return;
+			}
+			if (CommitType == ETextCommit::OnEnter && !InText.IsEmpty())
+			{
+				FName BusName(*InText.ToString());
+				UMovieSceneAnimationMixerTrack* Track = This->GetMixerTrack();
+				if (Track && FAnimMixerBusUtils::WouldBusSectionCreateCycle(BusName, Track, TracksForObject))
+				{
+					FNotificationInfo Info(FText::Format(
+						LOCTEXT("ChangeBusCycleBlocked", "Cannot read from bus '{0}': would create a dependency cycle."),
+						FText::FromName(BusName)));
+					Info.ExpireDuration = 4.0f;
+					FSlateNotificationManager::Get().AddNotification(Info);
+					return;
+				}
+				This->ChangeBusName(BusName);
+				FSlateApplication::Get().DismissAllMenus();
+			}
+		});
+
+	MenuBuilder.AddWidget(
+		SNew(SHorizontalBox)
+		+ SHorizontalBox::Slot()
+		.Padding(4.0f)
+		.FillWidth(1.0f)
+		[
+			TextBox
+		],
+		LOCTEXT("NewBusNameLabel", "Bus Name"),
+		/*bNoIndent=*/ true
+	);
+}
+
+void FAnimBusSectionInterface::ChangeBusName(FName NewBusName)
+{
+	UMovieSceneAnimBusSection* BusSection = GetBusSection();
+	if (!BusSection || NewBusName.IsNone())
+	{
+		return;
+	}
+
+	FScopedTransaction Transaction(LOCTEXT("ChangeBusName", "Change Bus Name"));
+	BusSection->Modify();
+	BusSection->BusName = NewBusName;
+}
+
+} // namespace UE::Sequencer
+
+#undef LOCTEXT_NAMESPACE
diff --git a/Engine/Plugins/Experimental/MovieSceneAnimMixer/Source/MovieSceneAnimMixerEditor/Private/AnimBusSectionInterface.h b/Engine/Plugins/Experimental/MovieSceneAnimMixer/Source/MovieSceneAnimMixerEditor/Private/AnimBusSectionInterface.h
new file mode 100644
--- /dev/null
+++ b/Engine/Plugins/Experimental/MovieSceneAnimMixer/Source/MovieSceneAnimMixerEditor/Private/AnimBusSectionInterface.h
@@ -0,0 +1,44 @@
+// Copyright Epic Games, Inc. All Rights Reserved.
+
+#pragma once
+
+#include "ISequencerSection.h"
+#include "Templates/SharedPointer.h"
+
+class UMovieSceneAnimBusSection;
+class UMovieSceneAnimationMixerTrack;
+class FMenuBuilder;
+class ISequencer;
+
+namespace UE::Sequencer
+{
+
+// Section interface for bus sections in Animation Mixer tracks.
+// Inherits from FSequencerSection for proper background rendering,
+// and adds bus-specific title and context menu.
+class FAnimBusSectionInterface
+	: public FSequencerSection
+	, public TSharedFromThis<FAnimBusSectionInterface>
+{
+public:
+	FAnimBusSectionInterface(UMovieSceneSection& InSection, UMovieSceneTrack& InTrack, FGuid InObjectBinding, TWeakPtr<ISequencer> InSequencer);
+
+	virtual FText GetSectionTitle() const override;
+	virtual FText GetSectionToolTip() const override;
+	virtual bool SectionIsResizable() const override { return true; }
+	virtual bool IsDeletable() const override { return true; }
+	virtual void BuildSectionContextMenu(FMenuBuilder& MenuBuilder, const FGuid& ObjectBinding) override;
+
+private:
+	void PopulateChangeBusMenu(FMenuBuilder& MenuBuilder);
+	void ChangeBusName(FName NewBusName);
+
+	UMovieSceneAnimBusSection* GetBusSection() const;
+	UMovieSceneAnimationMixerTrack* GetMixerTrack() const;
+
+	TWeakObjectPtr<UMovieSceneTrack> WeakTrack;
+	TWeakPtr<ISequencer> WeakSequencer;
+	FGuid ObjectBindingId;
+};
+
+} // namespace UE::Sequencer
diff --git a/Engine/Plugins/Experimental/MovieSceneAnimMixer/Source/MovieSceneAnimMixerEditor/Private/AnimBusSectionMenuProvider.cpp b/Engine/Plugins/Experimental/MovieSceneAnimMixer/Source/MovieSceneAnimMixerEditor/Private/AnimBusSectionMenuProvider.cpp
new file mode 100644
--- /dev/null
+++ b/Engine/Plugins/Experimental/MovieSceneAnimMixer/Source/MovieSceneAnimMixerEditor/Private/AnimBusSectionMenuProvider.cpp
@@ -0,0 +1,194 @@
+// Copyright Epic Games, Inc. All Rights Reserved.
+
+#include "AnimBusSectionMenuProvider.h"
+
+#include "AnimMixerBusUtils.h"
+#include "AnimMixerEditorBusUtils.h"
+#include "Framework/MultiBox/MultiBoxBuilder.h"
+#include "Framework/Notifications/NotificationManager.h"
+#include "ISequencer.h"
+#include "MovieSceneAnimBusSection.h"
+#include "MovieSceneAnimationMixerTrack.h"
+#include "MovieSceneSequence.h"
+#include "ScopedTransaction.h"
+#include "Systems/MovieSceneAnimMixerSystem.h"
+#include "Components/SkeletalMeshComponent.h"
+#include "EntitySystem/MovieSceneEntitySystemLinker.h"
+#include "GameFramework/Actor.h"
+#include "Widgets/Input/SEditableTextBox.h"
+#include "Widgets/Notifications/SNotificationList.h"
+
+#define LOCTEXT_NAMESPACE "AnimBusSectionMenuProvider"
+
+const UClass* FAnimBusSectionMenuProvider::GetHandledMixerItemClass() const
+{
+	return UMovieSceneAnimBusSection::StaticClass();
+}
+
+void FAnimBusSectionMenuProvider::PopulateAddMixerItemMenu(
+	FMenuBuilder& MenuBuilder,
+	TArray<FGuid> ObjectBindings,
+	UMovieSceneTrack* Track,
+	TSharedPtr<ISequencer> Sequencer,
+	int32 RowIndex)
+{
+	UMovieSceneAnimationMixerTrack* MixerTrack = Cast<UMovieSceneAnimationMixerTrack>(Track);
+	if (!MixerTrack || !Sequencer)
+	{
+		return;
+	}
+
+	// Gather known bus names from the sequence hierarchy
+	TArray<FName> KnownBuses;
+	if (UMovieSceneSequence* RootSequence = Sequencer->GetRootMovieSceneSequence())
+	{
+		KnownBuses = FAnimMixerBusUtils::GatherBusNamesFromSequence(RootSequence);
+	}
+
+	TWeakObjectPtr<UMovieSceneAnimationMixerTrack> WeakMixerTrack = MixerTrack;
+	TWeakPtr<ISequencer> WeakSequencer = Sequencer;
+
+	// Shared creation function used by both the known-bus entries and the new-bus text field
+	TFunction<void(FName)> CreateBusSection = [WeakMixerTrack, RowIndex, WeakSequencer](FName InBusName)
+	{
+		TSharedPtr<ISequencer> Seq = WeakSequencer.Pin();
+		UMovieSceneAnimationMixerTrack* MixTrack = WeakMixerTrack.Get();
+		if (!Seq || !MixTrack || InBusName.IsNone())
+		{
+			return;
+		}
+
+		// Check for cycles before modifying the track
+		TArray<UMovieSceneAnimationMixerTrack*> TracksForObject = AnimMixerEditorBusUtils::GatherMixerTracksForSameObject(MixTrack, *Seq);
+		if (FAnimMixerBusUtils::WouldBusSectionCreateCycle(InBusName, MixTrack, TracksForObject))
+		{
+			FNotificationInfo Info(LOCTEXT("BusCycleDetected", "Cannot add bus section: this would create a dependency cycle."));
+			Info.ExpireDuration = 4.0f;
+			FSlateNotificationManager::Get().AddNotification(Info);
+			return;
+		}
+
+		FScopedTransaction Transaction(LOCTEXT("AddBusSection", "Add Bus Section"));
+		MixTrack->Modify();
+
+		UMovieSceneAnimBusSection* NewSection = NewObject<UMovieSceneAnimBusSection>(MixTrack, NAME_None, RF_Transactional);
+		NewSection->BusName = InBusName;
+
+		// Determine the row index for the new section
+		int32 NewRowIndex = RowIndex;
+		if (NewRowIndex == INDEX_NONE)
+		{
+			int32 MaxRowIndex = -1;
+			for (UMovieSceneSection* Section : MixTrack->GetAllSections())
+			{
+				MaxRowIndex = FMath::Max(MaxRowIndex, Section->GetRowIndex());
+			}
+			NewRowIndex = FMath::Max(MaxRowIndex + 1, MixTrack->GetNextChildTrackRowIndex());
+		}
+		NewSection->SetRowIndex(NewRowIndex);
+
+		static_cast<UMovieSceneTrack*>(MixTrack)->AddSection(UMovieSceneTrack::FSectionParameter{*NewSection});
+
+		// Auto-add root motion settings if the bus source has root motion
+		if (TSharedPtr<UE::MovieScene::FSharedPlaybackState> PlaybackState = Seq->FindSharedPlaybackState())
+		{
+			if (UMovieSceneEntitySystemLinker* SeqLinker = PlaybackState->GetLinker())
+			{
+				if (UMovieSceneAnimMixerSystem* MixerSystem = SeqLinker->FindSystem<UMovieSceneAnimMixerSystem>())
+				{
+					FGuid BindingGuid = MixTrack->FindObjectBindingGuid();
+					TArrayView<TWeakObjectPtr<>> BoundObjects = Seq->FindObjectsInCurrentSequence(BindingGuid);
+					for (const TWeakObjectPtr<>& WeakObj : BoundObjects)
+					{
+						if (UObject* BoundObj = WeakObj.Get())
+						{
+							// Resolve to the skeletal mesh component, matching
+							// how the mixer system keys bus storage.
+							UObject* ResolvedObj = BoundObj;
+							if (AActor* Actor = Cast<AActor>(BoundObj))
+							{
+								if (USkeletalMeshComponent* SkelMesh = Actor->FindComponentByClass<USkeletalMeshComponent>())
+								{
+									ResolvedObj = SkelMesh;
+								}
+							}
+							TSharedPtr<FMovieSceneAnimBusData> BusData = MixerSystem->ReadBusData(FObjectKey(ResolvedObj), InBusName);
+							if (BusData && BusData->bHasRootMotion)
+							{
+								auto* RootMotionSettings = NewObject<UMovieSceneRootMotionSettingsDecoration>(
+									NewSection, UMovieSceneRootMotionSettingsDecoration::StaticClass(), NAME_None, RF_Transactional);
+								
+								// Default to world space since buses bake their root motion back to the root bone in world space
+								RootMotionSettings->RootMotionSpace.SetDefault((uint8)EMovieSceneRootMotionSpace::WorldSpace);
+								NewSection->AddDecoration(RootMotionSettings);
+								break;
+							}
+						}
+					}
+				}
+			}
+		}
+
+		Seq->NotifyMovieSceneDataChanged(EMovieSceneDataChangeType::MovieSceneStructureItemAdded);
+	};
+
+	UMovieSceneSequence* RootSequence = Sequencer ? Sequencer->GetRootMovieSceneSequence() : nullptr;
+
+	// Pre-compute tracks for this object so the submenu lambda can do cycle checks
+	TArray<UMovieSceneAnimationMixerTrack*> TracksForObject = Sequencer
+		? AnimMixerEditorBusUtils::GatherMixerTracksForSameObject(MixerTrack, *Sequencer)
+		: TArray<UMovieSceneAnimationMixerTrack*>();
+
+	MenuBuilder.AddSubMenu(
+		LOCTEXT("BusSectionMenu", "Bus"),
+		LOCTEXT("BusSectionMenuTooltip", "Add a bus section that reads from a named bus."),
+		FNewMenuDelegate::CreateLambda([KnownBuses, CreateBusSection, WeakMixerTrack, TracksForObject](FMenuBuilder& SubMenuBuilder)
+		{
+			UMovieSceneAnimationMixerTrack* MixTrack = WeakMixerTrack.Get();
+
+			for (const FName& BusName : KnownBuses)
+			{
+				const bool bWouldCycle = MixTrack &&
+					FAnimMixerBusUtils::WouldBusSectionCreateCycle(BusName, MixTrack, TracksForObject);
+
+				SubMenuBuilder.AddMenuEntry(
+					FText::FromName(BusName),
+					bWouldCycle
+						? FText::Format(LOCTEXT("BusSectionCycleTooltip", "Cannot read from bus '{0}': would create a dependency cycle."), FText::FromName(BusName))
+						: FText::Format(LOCTEXT("BusSectionEntryTooltip", "Read from bus '{0}'."), FText::FromName(BusName)),
+					FSlateIcon(),
+					FUIAction(
+						FExecuteAction::CreateLambda([CreateBusSection, BusName]() { CreateBusSection(BusName); }),
+						FCanExecuteAction::CreateLambda([bWouldCycle]() { return !bWouldCycle; }))
+				);
+			}
+
+			// "New Bus..." with inline text input
+			TSharedRef<SEditableTextBox> TextBox = SNew(SEditableTextBox)
+				.HintText(LOCTEXT("NewBusNameHint", "Enter bus name..."))
+				.OnTextCommitted_Lambda([CreateBusSection](const FText& InText, ETextCommit::Type CommitType)
+				{
+					if (CommitType == ETextCommit::OnEnter && !InText.IsEmpty())
+					{
+						FName BusName(*InText.ToString());
+						CreateBusSection(BusName);
+						FSlateApplication::Get().DismissAllMenus();
+					}
+				});
+
+			SubMenuBuilder.AddWidget(
+				SNew(SHorizontalBox)
+				+ SHorizontalBox::Slot()
+				.Padding(4.0f)
+				.FillWidth(1.0f)
+				[
+					TextBox
+				],
+				LOCTEXT("NewBusLabel", "New Bus..."),
+				/*bNoIndent=*/ true
+			);
+		})
+	);
+}
+
+#undef LOCTEXT_NAMESPACE
diff --git a/Engine/Plugins/Experimental/MovieSceneAnimMixer/Source/MovieSceneAnimMixerEditor/Private/AnimBusSectionMenuProvider.h b/Engine/Plugins/Experimental/MovieSceneAnimMixer/Source/MovieSceneAnimMixerEditor/Private/AnimBusSectionMenuProvider.h
new file mode 100644
--- /dev/null
+++ b/Engine/Plugins/Experimental/MovieSceneAnimMixer/Source/MovieSceneAnimMixerEditor/Private/AnimBusSectionMenuProvider.h
@@ -0,0 +1,15 @@
+// Copyright Epic Games, Inc. All Rights Reserved.
+
+#pragma once
+
+#include "IMovieSceneAnimMixerItemMenuProvider.h"
+
+// Menu provider for adding bus sections to Animation Mixer layers.
+// Shows a submenu with known bus names plus a "New Bus..." entry.
+class FAnimBusSectionMenuProvider : public IMovieSceneAnimMixerItemMenuProvider
+{
+public:
+	virtual const UClass* GetHandledMixerItemClass() const override;
+	virtual void PopulateAddMixerItemMenu(FMenuBuilder& MenuBuilder, TArray<FGuid> ObjectBindings, UMovieSceneTrack* Track, TSharedPtr<ISequencer> Sequencer, int32 RowIndex) override;
+	virtual int32 GetMixerItemMenuPriority() const override { return 50; }
+};
diff --git a/Engine/Plugins/Experimental/MovieSceneAnimMixer/Source/MovieSceneAnimMixerEditor/Private/AnimBusTargetMenuProvider.cpp b/Engine/Plugins/Experimental/MovieSceneAnimMixer/Source/MovieSceneAnimMixerEditor/Private/AnimBusTargetMenuProvider.cpp
new file mode 100644
--- /dev/null
+++ b/Engine/Plugins/Experimental/MovieSceneAnimMixer/Source/MovieSceneAnimMixerEditor/Private/AnimBusTargetMenuProvider.cpp
@@ -0,0 +1,112 @@
+// Copyright Epic Games, Inc. All Rights Reserved.
+
+#include "AnimBusTargetMenuProvider.h"
+
+#include "AnimMixerBusUtils.h"
+#include "AnimMixerEditorBusUtils.h"
+#include "Framework/Application/SlateApplication.h"
+#include "Framework/MultiBox/MultiBoxBuilder.h"
+#include "ISequencer.h"
+#include "MovieSceneAnimationMixerTrack.h"
+#include "MovieSceneMixedAnimationTarget.h"
+#include "MovieSceneSequence.h"
+#include "Systems/MovieSceneAnimMixerSystem.h"
+#include "Widgets/Input/SEditableTextBox.h"
+
+#define LOCTEXT_NAMESPACE "AnimBusTargetMenuProvider"
+
+UScriptStruct* FAnimBusTargetMenuProvider::GetHandledTargetStructType() const
+{
+	return FMovieSceneAnimBusTarget::StaticStruct();
+}
+
+void FAnimBusTargetMenuProvider::PopulateTargetMenu(
+	FMenuBuilder& MenuBuilder,
+	TObjectPtr<UObject> BoundObject,
+	TFunction<void(TInstancedStruct<FMovieSceneMixedAnimationTarget>)> OnTargetSelected,
+	TSharedPtr<ISequencer> Sequencer,
+	UMovieSceneAnimationMixerTrack* Track)
+{
+	MenuBuilder.AddSubMenu(
+		LOCTEXT("BusTarget", "Bus Target"),
+		LOCTEXT("BusTargetTooltip", "Write the mixer's evaluated pose to a named bus for consumption by other mixers."),
+		FNewMenuDelegate::CreateRaw(this, &FAnimBusTargetMenuProvider::PopulateBusTargetSubmenu, BoundObject, OnTargetSelected, Sequencer, Track)
+	);
+}
+
+void FAnimBusTargetMenuProvider::PopulateTargetMenu(
+	FMenuBuilder& MenuBuilder,
+	TObjectPtr<UObject> BoundObject,
+	TFunction<void(TInstancedStruct<FMovieSceneMixedAnimationTarget>)> OnTargetSelected)
+{
+	PopulateTargetMenu(MenuBuilder, BoundObject, OnTargetSelected, nullptr, nullptr);
+}
+
+void FAnimBusTargetMenuProvider::PopulateBusTargetSubmenu(
+	FMenuBuilder& MenuBuilder,
+	TObjectPtr<UObject> BoundObject,
+	TFunction<void(TInstancedStruct<FMovieSceneMixedAnimationTarget>)> OnTargetSelected,
+	TSharedPtr<ISequencer> Sequencer,
+	UMovieSceneAnimationMixerTrack* Track)
+{
+	UMovieSceneSequence* RootSequence = Sequencer ? Sequencer->GetRootMovieSceneSequence() : nullptr;
+
+	// List known buses
+	TArray<FName> KnownBuses;
+	if (RootSequence)
+	{
+		KnownBuses = FAnimMixerBusUtils::GatherBusNamesFromSequence(RootSequence);
+	}
+
+	TArray<UMovieSceneAnimationMixerTrack*> TracksForObject = Sequencer
+		? AnimMixerEditorBusUtils::GatherMixerTracksForSameObject(Track, *Sequencer)
+		: TArray<UMovieSceneAnimationMixerTrack*>();
+
+	for (const FName& BusName : KnownBuses)
+	{
+		const bool bWouldCycle = Track &&
+			FAnimMixerBusUtils::WouldBusTargetCreateCycle(BusName, Track, TracksForObject);
+
+		MenuBuilder.AddMenuEntry(
+			FText::FromName(BusName),
+			bWouldCycle
+				? FText::Format(LOCTEXT("BusTargetCycleTooltip", "Cannot target bus '{0}': would create a dependency cycle."), FText::FromName(BusName))
+				: FText::Format(LOCTEXT("BusTargetEntryTooltip", "Write to bus '{0}'."), FText::FromName(BusName)),
+			FSlateIcon(),
+			FUIAction(
+				FExecuteAction::CreateLambda([OnTargetSelected, BusName]()
+				{
+					OnTargetSelected(TInstancedStruct<FMovieSceneAnimBusTarget>::Make(BusName));
+					FSlateApplication::Get().DismissAllMenus();
+				}),
+				FCanExecuteAction::CreateLambda([bWouldCycle]() { return !bWouldCycle; }))
+		);
+	}
+
+	// "New Bus..." text entry
+	TSharedRef<SEditableTextBox> TextBox = SNew(SEditableTextBox)
+		.HintText(LOCTEXT("NewBusNameHint", "Enter bus name..."))
+		.OnTextCommitted_Lambda([OnTargetSelected](const FText& InText, ETextCommit::Type CommitType)
+		{
+			if (CommitType == ETextCommit::OnEnter && !InText.IsEmpty())
+			{
+				FName BusName(*InText.ToString());
+				OnTargetSelected(TInstancedStruct<FMovieSceneAnimBusTarget>::Make(BusName));
+				FSlateApplication::Get().DismissAllMenus();
+			}
+		});
+
+	MenuBuilder.AddWidget(
+		SNew(SHorizontalBox)
+		+ SHorizontalBox::Slot()
+		.Padding(4.0f)
+		.FillWidth(1.0f)
+		[
+			TextBox
+		],
+		LOCTEXT("NewBusLabel", "New Bus..."),
+		/*bNoIndent=*/ true
+	);
+}
+
+#undef LOCTEXT_NAMESPACE
diff --git a/Engine/Plugins/Experimental/MovieSceneAnimMixer/Source/MovieSceneAnimMixerEditor/Private/AnimBusTargetMenuProvider.h b/Engine/Plugins/Experimental/MovieSceneAnimMixer/Source/MovieSceneAnimMixerEditor/Private/AnimBusTargetMenuProvider.h
new file mode 100644
--- /dev/null
+++ b/Engine/Plugins/Experimental/MovieSceneAnimMixer/Source/MovieSceneAnimMixerEditor/Private/AnimBusTargetMenuProvider.h
@@ -0,0 +1,37 @@
+// Copyright Epic Games, Inc. All Rights Reserved.
+
+#pragma once
+
+#include "IMovieSceneAnimMixerTargetMenuProvider.h"
+
+// Target menu provider for the Bus Target type.
+// Shows a submenu with known bus names plus a "New Bus..." entry.
+class FAnimBusTargetMenuProvider : public IMovieSceneAnimMixerTargetMenuProvider
+{
+public:
+	virtual UScriptStruct* GetHandledTargetStructType() const override;
+
+	// Override the 5-param version to use sequencer/track for bus listing and cycle checks
+	virtual void PopulateTargetMenu(
+		FMenuBuilder& MenuBuilder,
+		TObjectPtr<UObject> BoundObject,
+		TFunction<void(TInstancedStruct<FMovieSceneMixedAnimationTarget>)> OnTargetSelected,
+		TSharedPtr<ISequencer> Sequencer,
+		UMovieSceneAnimationMixerTrack* Track) override;
+
+	// 3-param fallback
+	virtual void PopulateTargetMenu(
+		FMenuBuilder& MenuBuilder,
+		TObjectPtr<UObject> BoundObject,
+		TFunction<void(TInstancedStruct<FMovieSceneMixedAnimationTarget>)> OnTargetSelected) override;
+
+	virtual int32 GetTargetMenuPriority() const override { return 40; }
+
+private:
+	void PopulateBusTargetSubmenu(
+		FMenuBuilder& MenuBuilder,
+		TObjectPtr<UObject> BoundObject,
+		TFunction<void(TInstancedStruct<FMovieSceneMixedAnimationTarget>)> OnTargetSelected,
+		TSharedPtr<ISequencer> Sequencer,
+		UMovieSceneAnimationMixerTrack* Track);
+};
diff --git a/Engine/Plugins/Experimental/MovieSceneAnimMixer/Source/MovieSceneAnimMixerEditor/Private/AnimMixerEditorBusUtils.cpp b/Engine/Plugins/Experimental/MovieSceneAnimMixer/Source/MovieSceneAnimMixerEditor/Private/AnimMixerEditorBusUtils.cpp
new file mode 100644
--- /dev/null
+++ b/Engine/Plugins/Experimental/MovieSceneAnimMixer/Source/MovieSceneAnimMixerEditor/Private/AnimMixerEditorBusUtils.cpp
@@ -0,0 +1,191 @@
+// Copyright Epic Games, Inc. All Rights Reserved.
+
+#include "AnimMixerEditorBusUtils.h"
+
+#include "ISequencer.h"
+#include "MovieSceneAnimationMixerTrack.h"
+#include "Compilation/MovieSceneCompiledDataManager.h"
+#include "Evaluation/MovieSceneEvaluationTemplateInstance.h"
+#include "Evaluation/MovieSceneSequenceHierarchy.h"
+#include "MovieScene.h"
+#include "MovieSceneSequence.h"
+
+namespace AnimMixerEditorBusUtils
+{
+
+// Collect mixer tracks from a movie scene whose binding resolves to TargetObject.
+static void GatherFromMovieScene(
+	UMovieScene* MovieScene,
+	FMovieSceneSequenceID SequenceID,
+	UObject* TargetObject,
+	ISequencer& Sequencer,
+	TArray<UMovieSceneAnimationMixerTrack*>& OutTracks)
+{
+	if (!MovieScene)
+	{
+		return;
+	}
+
+	const UMovieScene* ConstMovieScene = MovieScene;
+	for (const FMovieSceneBinding& Binding : ConstMovieScene->GetBindings())
+	{
+		bool bMatchesTarget = false;
+		TArrayView<TWeakObjectPtr<>> BoundObjects = Sequencer.FindBoundObjects(Binding.GetObjectGuid(), SequenceID);
+		for (TWeakObjectPtr<> WeakObj : BoundObjects)
+		{
+			if (WeakObj.Get() == TargetObject)
+			{
+				bMatchesTarget = true;
+				break;
+			}
+		}
+
+		if (!bMatchesTarget)
+		{
+			continue;
+		}
+
+		for (UMovieSceneTrack* Track : Binding.GetTracks())
+		{
+			if (UMovieSceneAnimationMixerTrack* MixerTrack = Cast<UMovieSceneAnimationMixerTrack>(Track))
+			{
+				OutTracks.AddUnique(MixerTrack);
+			}
+		}
+	}
+}
+
+// Find the SequenceID for a movie scene in the compiled hierarchy.
+// Returns MovieSceneSequenceID::Invalid if not found.
+static FMovieSceneSequenceID FindSequenceIDForMovieScene(
+	UMovieScene* TargetMovieScene,
+	ISequencer& Sequencer)
+{
+	UMovieSceneSequence* RootSequence = Sequencer.GetRootMovieSceneSequence();
+	if (RootSequence && RootSequence->GetMovieScene() == TargetMovieScene)
+	{
+		return MovieSceneSequenceID::Root;
+	}
+
+	FMovieSceneRootEvaluationTemplateInstance& RootInstance = Sequencer.GetEvaluationTemplate();
+	UMovieSceneCompiledDataManager* CompiledDataManager = RootInstance.GetCompiledDataManager();
+	if (!CompiledDataManager)
+	{
+		return MovieSceneSequenceID::Invalid;
+	}
+
+	const FMovieSceneSequenceHierarchy* Hierarchy = CompiledDataManager->FindHierarchy(RootInstance.GetCompiledDataID());
+	if (!Hierarchy)
+	{
+		return MovieSceneSequenceID::Invalid;
+	}
+
+	for (const auto& Pair : Hierarchy->AllSubSequenceData())
+	{
+		UMovieSceneSequence* SubSequence = Pair.Value.GetSequence();
+		if (SubSequence && SubSequence->GetMovieScene() == TargetMovieScene)
+		{
+			return Pair.Key;
+		}
+	}
+
+	return MovieSceneSequenceID::Invalid;
+}
+
+TArray<UMovieSceneAnimationMixerTrack*> GatherMixerTracksForSameObject(
+	UMovieSceneAnimationMixerTrack* Track, ISequencer& Sequencer)
+{
+	if (!Track)
+	{
+		return {};
+	}
+
+	UMovieScene* TrackMovieScene = Track->GetTypedOuter<UMovieScene>();
+	if (!TrackMovieScene)
+	{
+		return {};
+	}
+
+	// Find the track's binding GUID
+	FGuid BindingGuid;
+	const UMovieScene* ConstMovieScene = TrackMovieScene;
+	for (const FMovieSceneBinding& Binding : ConstMovieScene->GetBindings())
+	{
+		bool bFound = false;
+		for (UMovieSceneTrack* BoundTrack : Binding.GetTracks())
+		{
+			if (BoundTrack == Track)
+			{
+				BindingGuid = Binding.GetObjectGuid();
+				bFound = true;
+				break;
+			}
+		}
+		if (bFound)
+		{
+			break;
+		}
+	}
+
+	if (!BindingGuid.IsValid())
+	{
+		return {};
+	}
+
+	// Find the SequenceID for the track's movie scene so we resolve
+	// the binding in the correct sub-sequence context
+	FMovieSceneSequenceID TrackSequenceID = FindSequenceIDForMovieScene(TrackMovieScene, Sequencer);
+	if (TrackSequenceID == MovieSceneSequenceID::Invalid)
+	{
+		return {};
+	}
+
+	// Resolve the binding to the actual bound object
+	TArrayView<TWeakObjectPtr<>> BoundObjects = Sequencer.FindBoundObjects(BindingGuid, TrackSequenceID);
+	UObject* BoundObject = nullptr;
+	for (TWeakObjectPtr<> WeakObj : BoundObjects)
+	{
+		BoundObject = WeakObj.Get();
+		if (BoundObject)
+		{
+			break;
+		}
+	}
+
+	if (!BoundObject)
+	{
+		return {};
+	}
+
+	// Now gather all mixer tracks across the full hierarchy that resolve to
+	// the same bound object
+	TArray<UMovieSceneAnimationMixerTrack*> Result;
+
+	UMovieSceneSequence* RootSequence = Sequencer.GetRootMovieSceneSequence();
+	if (RootSequence)
+	{
+		GatherFromMovieScene(RootSequence->GetMovieScene(), MovieSceneSequenceID::Root, BoundObject, Sequencer, Result);
+	}
+
+	FMovieSceneRootEvaluationTemplateInstance& RootInstance = Sequencer.GetEvaluationTemplate();
+	UMovieSceneCompiledDataManager* CompiledDataManager = RootInstance.GetCompiledDataManager();
+	if (CompiledDataManager)
+	{
+		const FMovieSceneSequenceHierarchy* Hierarchy = CompiledDataManager->FindHierarchy(RootInstance.GetCompiledDataID());
+		if (Hierarchy)
+		{
+			for (const auto& Pair : Hierarchy->AllSubSequenceData())
+			{
+				UMovieSceneSequence* SubSequence = Pair.Value.GetSequence();
+				if (SubSequence)
+				{
+					GatherFromMovieScene(SubSequence->GetMovieScene(), Pair.Key, BoundObject, Sequencer, Result);
+				}
+			}
+		}
+	}
+
+	return Result;
+}
+
+} // namespace AnimMixerEditorBusUtils
diff --git a/Engine/Plugins/Experimental/MovieSceneAnimMixer/Source/MovieSceneAnimMixerEditor/Private/AnimMixerEditorBusUtils.h b/Engine/Plugins/Experimental/MovieSceneAnimMixer/Source/MovieSceneAnimMixerEditor/Private/AnimMixerEditorBusUtils.h
new file mode 100644
--- /dev/null
+++ b/Engine/Plugins/Experimental/MovieSceneAnimMixer/Source/MovieSceneAnimMixerEditor/Private/AnimMixerEditorBusUtils.h
@@ -0,0 +1,19 @@
+// Copyright Epic Games, Inc. All Rights Reserved.
+
+#pragma once
+
+#include "CoreMinimal.h"
+
+class ISequencer;
+class UMovieSceneAnimationMixerTrack;
+
+// Editor-side bus utilities that require ISequencer for binding resolution.
+namespace AnimMixerEditorBusUtils
+{
+	// Gather all mixer tracks across the sequence hierarchy that are bound to
+	// the same object as the given track. Resolves the track's binding through
+	// the compiled hierarchy (not just the focused sequence) so this works
+	// correctly for tracks in any sub-sequence.
+	TArray<UMovieSceneAnimationMixerTrack*> GatherMixerTracksForSameObject(
+		UMovieSceneAnimationMixerTrack* Track, ISequencer& Sequencer);
+}
diff --git a/Engine/Plugins/Experimental/MovieSceneAnimMixer/Source/MovieSceneAnimMixerEditor/Private/MovieSceneAnimMixerEditorModule.cpp b/Engine/Plugins/Experimental/MovieSceneAnimMixer/Source/MovieSceneAnimMixerEditor/Private/MovieSceneAnimMixerEditorModule.cpp
--- a/Engine/Plugins/Experimental/MovieSceneAnimMixer/Source/MovieSceneAnimMixerEditor/Private/MovieSceneAnimMixerEditorModule.cpp
+++ b/Engine/Plugins/Experimental/MovieSceneAnimMixer/Source/MovieSceneAnimMixerEditor/Private/MovieSceneAnimMixerEditorModule.cpp
@@ -8,7 +8,11 @@
 #include "LayerWeightDecorationMenuProvider.h"
 #include "IMovieSceneRootMotionOffsetProvider.h"
 #include "AnimMixerTargetMenuProviders.h"
+#include "AnimBusTargetMenuProvider.h"
+#include "AnimBusSectionMenuProvider.h"
+#include "AnimBusSectionInterface.h"
 #include "AnimTransitionSectionInterface.h"
+#include "MovieSceneAnimBusSection.h"
 #include "IAnimGraph_SequencerMixerTargetConnector.h"
 #include "IMovieSceneAnimMixerTargetMenuProvider.h"
 #include "ISequencerDecorationEditor.h"
@@ -223,7 +227,7 @@ namespace UE::MovieScene
 		SequencerModule.RegisterChannelInterface<FMovieSceneByteChannelDefaultOnly>();
 		SequencerModule.RegisterChannelInterface<FMovieSceneBoneMatchChannel>();
 
-		// Register custom section interface for transition sections
+		// Register custom section interfaces for transition sections
 		UE::Sequencer::FAnimationMixerTrackEditor::RegisterCustomMixerAnimSection(
 			UMovieSceneAnimCrossfadeTransitionSection::StaticClass(),
 			FOnMakeSectionInterfaceDelegate::CreateLambda([](UMovieSceneSection& Section, UMovieSceneTrack& Track, FGuid ObjectBinding) -> TSharedRef<ISequencerSection>
@@ -249,6 +253,13 @@ namespace UE::MovieScene
 		IModularFeatures::Get().RegisterModularFeature(IMovieSceneAnimMixerTargetMenuProvider::GetModularFeatureName(), AnimBlueprintTargetMenuProvider.Get());
 		IModularFeatures::Get().RegisterModularFeature(IMovieSceneAnimMixerTargetMenuProvider::GetModularFeatureName(), AnimNextInjectionTargetMenuProvider.Get());
 
+		BusTargetMenuProvider = MakeUnique<FAnimBusTargetMenuProvider>();
+		IModularFeatures::Get().RegisterModularFeature(IMovieSceneAnimMixerTargetMenuProvider::GetModularFeatureName(), BusTargetMenuProvider.Get());
+
+		// Register bus section menu provider
+		BusSectionMenuProvider = MakeUnique<FAnimBusSectionMenuProvider>();
+		IModularFeatures::Get().RegisterModularFeature(IMovieSceneAnimMixerItemMenuProvider::GetModularFeatureName(), BusSectionMenuProvider.Get());
+
 		// Register decoration editors
 		RootMotionTargetDecorationEditor = MakeUnique<UE::Sequencer::FRootMotionTargetDecorationEditor>();
 		IModularFeatures::Get().RegisterModularFeature(ISequencerDecorationEditor::GetModularFeatureName(), RootMotionTargetDecorationEditor.Get());
@@ -321,6 +332,16 @@ namespace UE::MovieScene
 		{
 			IModularFeatures::Get().UnregisterModularFeature(IMovieSceneAnimMixerTargetMenuProvider::GetModularFeatureName(), AnimNextInjectionTargetMenuProvider.Get());
 		}
+		if (BusTargetMenuProvider)
+		{
+			IModularFeatures::Get().UnregisterModularFeature(IMovieSceneAnimMixerTargetMenuProvider::GetModularFeatureName(), BusTargetMenuProvider.Get());
+		}
+
+		// Unregister bus section menu provider
+		if (BusSectionMenuProvider)
+		{
+			IModularFeatures::Get().UnregisterModularFeature(IMovieSceneAnimMixerItemMenuProvider::GetModularFeatureName(), BusSectionMenuProvider.Get());
+		}
 
 		// Unregister custom section interface for transition sections
 		UE::Sequencer::FAnimationMixerTrackEditor::UnregisterCustomMixerAnimSection(UMovieSceneAnimCrossfadeTransitionSection::StaticClass());
diff --git a/Engine/Plugins/Experimental/MovieSceneAnimMixer/Source/MovieSceneAnimMixerEditor/Private/MovieSceneAnimMixerEditorModule.h b/Engine/Plugins/Experimental/MovieSceneAnimMixer/Source/MovieSceneAnimMixerEditor/Private/MovieSceneAnimMixerEditorModule.h
--- a/Engine/Plugins/Experimental/MovieSceneAnimMixer/Source/MovieSceneAnimMixerEditor/Private/MovieSceneAnimMixerEditorModule.h
+++ b/Engine/Plugins/Experimental/MovieSceneAnimMixer/Source/MovieSceneAnimMixerEditor/Private/MovieSceneAnimMixerEditorModule.h
@@ -12,6 +12,8 @@ class FAutomaticTargetMenuProvider;
 class FAnimInstanceTargetMenuProvider;
 class FAnimBlueprintTargetMenuProvider;
 class FAnimNextInjectionTargetMenuProvider;
+class FAnimBusTargetMenuProvider;
+class FAnimBusSectionMenuProvider;
 class FMaskDecorationMenuProvider;
 class FLayerWeightDecorationMenuProvider;
 
@@ -49,6 +51,10 @@ namespace UE::MovieScene
 		TUniquePtr<FAnimInstanceTargetMenuProvider> AnimInstanceTargetMenuProvider;
 		TUniquePtr<FAnimBlueprintTargetMenuProvider> AnimBlueprintTargetMenuProvider;
 		TUniquePtr<FAnimNextInjectionTargetMenuProvider> AnimNextInjectionTargetMenuProvider;
+		TUniquePtr<FAnimBusTargetMenuProvider> BusTargetMenuProvider;
+
+		// Bus section menu provider
+		TUniquePtr<FAnimBusSectionMenuProvider> BusSectionMenuProvider;
 
 		// Decoration editors
 		TUniquePtr<UE::Sequencer::FRootMotionTargetDecorationEditor> RootMotionTargetDecorationEditor;
diff --git a/Engine/Plugins/Experimental/MovieSceneAnimMixer/Source/MovieSceneAnimMixerEditor/Private/MovieSceneAnimationMixerTrackEditor.cpp b/Engine/Plugins/Experimental/MovieSceneAnimMixer/Source/MovieSceneAnimMixerEditor/Private/MovieSceneAnimationMixerTrackEditor.cpp
--- a/Engine/Plugins/Experimental/MovieSceneAnimMixer/Source/MovieSceneAnimMixerEditor/Private/MovieSceneAnimationMixerTrackEditor.cpp
+++ b/Engine/Plugins/Experimental/MovieSceneAnimMixer/Source/MovieSceneAnimMixerEditor/Private/MovieSceneAnimationMixerTrackEditor.cpp
@@ -50,6 +50,12 @@
 #include "TimeToPixel.h"
 #include "Widgets/Input/SComboButton.h"
 #include "Widgets/Images/SImage.h"
+#include "Framework/Notifications/NotificationManager.h"
+#include "Widgets/Notifications/SNotificationList.h"
+#include "AnimBusSectionInterface.h"
+#include "AnimMixerBusUtils.h"
+#include "AnimMixerEditorBusUtils.h"
+#include "MovieSceneAnimBusSection.h"
 #include "MVVM/AnimationMixerTrackModel.h"
 
 #define LOCTEXT_NAMESPACE "MovieSceneAnimationMixerTrackEditor"
@@ -142,11 +148,12 @@ private:
 		};
 
 		// Populate menu from all providers
+		UMovieSceneAnimationMixerTrack* Track = WeakTrack.Get();
 		for (IMovieSceneAnimMixerTargetMenuProvider* Provider : Providers)
 		{
 			if (Provider)
 			{
-				Provider->PopulateTargetMenu(MenuBuilder, BoundObject, OnTargetSelected);
+				Provider->PopulateTargetMenu(MenuBuilder, BoundObject, OnTargetSelected, Sequencer, Track);
 			}
 		}
 
@@ -194,7 +201,24 @@ private:
 		{
 			FScopedTransaction Transaction(LOCTEXT("SetAnimMixerTarget", "Set Animation Mixer Target"));
 			Track->Modify();
+			TInstancedStruct<FMovieSceneMixedAnimationTarget> PreviousTarget = Track->MixedAnimationTarget;
 			Track->MixedAnimationTarget = MoveTemp(NewTarget);
+
+			{
+				FAnimMixerBusValidationResult Result = FAnimMixerBusUtils::ValidateBusTopology(
+					AnimMixerEditorBusUtils::GatherMixerTracksForSameObject(Track, *Sequencer));
+				if (Result.HasErrors())
+				{
+					Transaction.Cancel();
+					Track->MixedAnimationTarget = MoveTemp(PreviousTarget);
+
+					FNotificationInfo Info(FText::FromString(Result.Errors[0]));
+					Info.ExpireDuration = 4.0f;
+					FSlateNotificationManager::Get().AddNotification(Info);
+					return;
+				}
+			}
+
 			Sequencer->NotifyMovieSceneDataChanged(EMovieSceneDataChangeType::TrackValueChanged);
 		}
 	}
@@ -606,7 +630,7 @@ void FAnimationMixerTrackEditor::BuildTargetSelectionMenu(FMenuBuilder& MenuBuil
 	{
 		if (Provider)
 		{
-			Provider->PopulateTargetMenu(MenuBuilder, BoundObject, OnTargetSelected);
+			Provider->PopulateTargetMenu(MenuBuilder, BoundObject, OnTargetSelected, Sequencer, nullptr);
 		}
 	}
 }
@@ -625,18 +649,33 @@ void FAnimationMixerTrackEditor::HandleAddMixerTrackWithTarget(TArray<FGuid> Obj
 		return;
 	}
 
-	const FScopedTransaction Transaction(LOCTEXT("AddMixerTrack_Transaction", "Add Mixer Track"));
+	FScopedTransaction Transaction(LOCTEXT("AddMixerTrack_Transaction", "Add Mixer Track"));
 
 	MovieScene->Modify();
 
 	for (const FGuid& Guid : ObjectBindings)
 	{
-		FFindOrCreateTrackResult TrackResult = FindOrCreateTrackForObject(Guid, UMovieSceneAnimationMixerTrack::StaticClass());
-		if (TrackResult.bWasCreated)
+		UMovieSceneAnimationMixerTrack* MixerTrack = MovieScene->AddTrack<UMovieSceneAnimationMixerTrack>(Guid);
+		if (MixerTrack)
 		{
-			UMovieSceneAnimationMixerTrack* MixerTrack = CastChecked<UMovieSceneAnimationMixerTrack>(TrackResult.Track);
 			MixerTrack->SetDisplayName(MixerTrack->GetDefaultDisplayName());
 			MixerTrack->MixedAnimationTarget = Target;
+
+			{
+				FAnimMixerBusValidationResult Result = FAnimMixerBusUtils::ValidateBusTopology(
+					AnimMixerEditorBusUtils::GatherMixerTracksForSameObject(MixerTrack, *GetSequencer()));
+				if (Result.HasErrors())
+				{
+					Transaction.Cancel();
+
+					FNotificationInfo Info(FText::FromString(Result.Errors[0]));
+					Info.ExpireDuration = 4.0f;
+					FSlateNotificationManager::Get().AddNotification(Info);
+					return;
+				}
+			}
+
+			GetSequencer()->NotifyMovieSceneDataChanged(EMovieSceneDataChangeType::MovieSceneStructureItemAdded);
 		}
 	}
 }
@@ -746,6 +785,11 @@ TSharedRef<ISequencerSection> FAnimationMixerTrackEditor::MakeSectionInterface(U
 		return MakeShared<FAnimMixerAnimationSection>(SectionObject, GetSequencer());
 	}
 
+	if (SectionObject.IsA<UMovieSceneAnimBusSection>())
+	{
+		return MakeShared<FAnimBusSectionInterface>(SectionObject, Track, ObjectBinding, GetSequencer());
+	}
+
 	if (FOnMakeSectionInterfaceDelegate* MakeSectionInterfaceDelegate = MakeSectionInterfaceCallbacks.Find(SectionObject.GetClass()))
 	{
 		if ((*MakeSectionInterfaceDelegate).IsBound())
diff --git a/Engine/Source/Editor/MovieSceneTools/Public/IMovieSceneAnimMixerTargetMenuProvider.h b/Engine/Source/Editor/MovieSceneTools/Public/IMovieSceneAnimMixerTargetMenuProvider.h
--- a/Engine/Source/Editor/MovieSceneTools/Public/IMovieSceneAnimMixerTargetMenuProvider.h
+++ b/Engine/Source/Editor/MovieSceneTools/Public/IMovieSceneAnimMixerTargetMenuProvider.h
@@ -8,6 +8,8 @@
 #include "StructUtils/InstancedStruct.h"
 
 class FMenuBuilder;
+class ISequencer;
+class UMovieSceneAnimationMixerTrack;
 
 /**
  * Modular feature interface for providing custom UI when selecting animation mixer targets.
@@ -37,6 +39,16 @@ public:
 	 *
 	 * OnTargetSelected should be invoked by the menu provider with a constructed target once a target is selected.
 	 */
+	virtual void PopulateTargetMenu(
+		FMenuBuilder& MenuBuilder,
+		TObjectPtr<UObject> BoundObject,
+		TFunction<void(TInstancedStruct<FMovieSceneMixedAnimationTarget>)> OnTargetSelected,
+		TSharedPtr<ISequencer> Sequencer,
+		UMovieSceneAnimationMixerTrack* Track)
+	{
+		PopulateTargetMenu(MenuBuilder, BoundObject, OnTargetSelected);
+	}
+
 	virtual void PopulateTargetMenu(
 		FMenuBuilder& MenuBuilder,
 		TObjectPtr<UObject> BoundObject,
